Merging CL #9451677 from //UE4/Dev-Enterprise

Added default values for tessellation and import in BaseDatasmithImporter.ini
Forced import options for Datasmith file producer to be those from BaseDatasmithImporter.ini
Changed logic in GetSceneImportOptions and SetSceneImportOptions for DatasmithCADTranslator, DatasmithWireTranslator and DatasmithOpenNurbsTranslator.

#rnx

diff --git a/Engine/Plugins/Enterprise/DatasmithCADImporter/Source/DatasmithCADTranslator/Private/DatasmithCADTranslator.cpp b/Engine/Plugins/Enterprise/DatasmithCADImporter/Source/DatasmithCADTranslator/Private/DatasmithCADTranslator.cpp
--- a/Engine/Plugins/Enterprise/DatasmithCADImporter/Source/DatasmithCADTranslator/Private/DatasmithCADTranslator.cpp
+++ b/Engine/Plugins/Enterprise/DatasmithCADImporter/Source/DatasmithCADTranslator/Private/DatasmithCADTranslator.cpp
@@ -5,7 +5,6 @@
 #include "CADLibraryOptions.h"
 #include "CoreTechParametricSurfaceExtension.h"
 #include "DatasmithCADTranslatorModule.h"
-#include "DatasmithImportOptions.h"
 #include "DatasmithMeshHelper.h"
 #include "DatasmithSceneFactory.h"
 #include "DatasmithUtils.h"
@@ -89,7 +88,7 @@ bool FDatasmithCADTranslator::LoadScene(TSharedRef<IDatasmithScene> DatasmithSce
 	IFileManager::Get().MakeDirectory(*OutputPath, true);
 
 	Translator->SetOutputPath(OutputPath);
-	Translator->SetTessellationOptions(TessellationOptions);
+	Translator->SetTessellationOptions(GetCommonTessellationOptions());
 
 	return Translator->Read();
 
@@ -129,7 +128,7 @@ bool FDatasmithCADTranslator::LoadStaticMesh(const TSharedRef<IDatasmithMeshElem
 				CoreTechData->MeshParameters.SymmetricNormal = MeshParameters.SymmetricNormal;
 				CoreTechData->MeshParameters.SymmetricOrigin = MeshParameters.SymmetricOrigin;
 
-				CoreTechData->LastTessellationOptions = TessellationOptions;
+				CoreTechData->LastTessellationOptions = GetCommonTessellationOptions();
 				OutMeshPayload.AdditionalData.Add(CoreTechData);
 			}
 		}
@@ -137,27 +136,13 @@ bool FDatasmithCADTranslator::LoadStaticMesh(const TSharedRef<IDatasmithMeshElem
 	return OutMeshPayload.LodMeshes.Num() > 0;
 }
 
-void FDatasmithCADTranslator::GetSceneImportOptions(TArray<TStrongObjectPtr<UObject>>& Options)
-{
-	TStrongObjectPtr<UDatasmithCommonTessellationOptions> TessellationOptionsPtr = Datasmith::MakeOptions<UDatasmithCommonTessellationOptions>();
-	Options.Add(TessellationOptionsPtr);
-}
-
 void FDatasmithCADTranslator::SetSceneImportOptions(TArray<TStrongObjectPtr<UObject>>& Options)
 {
-	for (const TStrongObjectPtr<UObject>& OptionPtr : Options)
+	FDatasmithCoreTechTranslator::SetSceneImportOptions(Options);
+
+	if (Translator)
 	{
-		if (UObject* Option = OptionPtr.Get())
-		{
-			if (UDatasmithCommonTessellationOptions* TessellationOptionsObject = Cast<UDatasmithCommonTessellationOptions>(Option))
-			{
-				TessellationOptions = TessellationOptionsObject->Options;
-				if (Translator)
-				{
-					Translator->SetTessellationOptions(TessellationOptions);
-				}
-			}
-		}
+		Translator->SetTessellationOptions(GetCommonTessellationOptions());
 	}
 }
 
diff --git a/Engine/Plugins/Enterprise/DatasmithCADImporter/Source/DatasmithCADTranslator/Private/DatasmithCADTranslator.h b/Engine/Plugins/Enterprise/DatasmithCADImporter/Source/DatasmithCADTranslator/Private/DatasmithCADTranslator.h
--- a/Engine/Plugins/Enterprise/DatasmithCADImporter/Source/DatasmithCADTranslator/Private/DatasmithCADTranslator.h
+++ b/Engine/Plugins/Enterprise/DatasmithCADImporter/Source/DatasmithCADTranslator/Private/DatasmithCADTranslator.h
@@ -5,19 +5,17 @@
 #include "CoreMinimal.h"
 
 #include "CADLibraryOptions.h"
-#include "DatasmithImportOptions.h"
 #include "DatasmithCADTranslatorImpl.h"
-#include "Translators/DatasmithTranslator.h"
+#include "DatasmithCoreTechTranslator.h"
 #include "UObject/ObjectMacros.h"
 
-class FDatasmithCADTranslatorImpl;
 
-
-class FDatasmithCADTranslator : public IDatasmithTranslator
+class FDatasmithCADTranslator : public FDatasmithCoreTechTranslator
 {
 public:
 	FDatasmithCADTranslator();
 
+	// Begin IDatasmithTranslator overrides
 	virtual FName GetFName() const override { return "DatasmithCADTranslator"; };
 
 	virtual void Initialize(FDatasmithTranslatorCapabilities& OutCapabilities) override;
@@ -28,12 +26,13 @@ public:
 	virtual void UnloadScene() override;
 
 	virtual bool LoadStaticMesh(const TSharedRef<IDatasmithMeshElement> MeshElement, FDatasmithMeshElementPayload& OutMeshPayload) override;
+	// End IDatasmithTranslator overrides
 
-	virtual void GetSceneImportOptions(TArray<TStrongObjectPtr<UObject>>& Options) override;
+	// Begin ADatasmithCoreTechTranslator overrides
 	virtual void SetSceneImportOptions(TArray<TStrongObjectPtr<UObject>>& Options) override;
+	// End ADatasmithCoreTechTranslator overrides
 
 private:
 	TSharedPtr<FDatasmithCADTranslatorImpl> Translator;
-	FDatasmithTessellationOptions TessellationOptions;
 };
 
diff --git a/Engine/Plugins/Enterprise/DatasmithCADImporter/Source/DatasmithCADTranslator/Private/DatasmithCADTranslatorImpl.cpp b/Engine/Plugins/Enterprise/DatasmithCADImporter/Source/DatasmithCADTranslator/Private/DatasmithCADTranslatorImpl.cpp
--- a/Engine/Plugins/Enterprise/DatasmithCADImporter/Source/DatasmithCADTranslator/Private/DatasmithCADTranslatorImpl.cpp
+++ b/Engine/Plugins/Enterprise/DatasmithCADImporter/Source/DatasmithCADTranslator/Private/DatasmithCADTranslatorImpl.cpp
@@ -30,8 +30,6 @@ void FDatasmithCADTranslatorImpl::SetTessellationOptions(const FDatasmithTessell
 
 bool FDatasmithCADTranslatorImpl::Read()
 {
-	CTParser->SetTessellationOptions(TessellationOptions);
-
 #ifdef USE_CORETECH_MT_PARSER
 	CTParser->SetOutputPath(OutputPath);
 	CTParser->Read();
@@ -51,6 +49,10 @@ bool FDatasmithCADTranslatorImpl::Read()
 	}
 #endif
 
+	// Force CoreTech to re-tessellate the model with the translator's tessellation parameters
+	// This call has no effect on the load of the model
+	CTParser->SetTessellationOptions(TessellationOptions);
+
 	return true;
 }
 
diff --git a/Engine/Plugins/Enterprise/DatasmithCADImporter/Source/DatasmithCoreTechExtension/Private/DatasmithCoreTechTranslator.cpp b/Engine/Plugins/Enterprise/DatasmithCADImporter/Source/DatasmithCoreTechExtension/Private/DatasmithCoreTechTranslator.cpp
new file mode 100644
--- /dev/null
+++ b/Engine/Plugins/Enterprise/DatasmithCADImporter/Source/DatasmithCoreTechExtension/Private/DatasmithCoreTechTranslator.cpp
@@ -0,0 +1,40 @@
+// Copyright 1998-2019 Epic Games, Inc. All Rights Reserved.
+
+#include "DatasmithCoreTechTranslator.h"
+
+#include "DatasmithImportOptions.h"
+
+void FDatasmithCoreTechTranslator::GetSceneImportOptions(TArray<TStrongObjectPtr<UObject>>& Options)
+{
+	Options.Add(GetCommonTessellationOptionsPtr());
+}
+
+void FDatasmithCoreTechTranslator::SetSceneImportOptions(TArray<TStrongObjectPtr<UObject>>& Options)
+{
+	for (const TStrongObjectPtr<UObject>& OptionPtr : Options)
+	{
+		if (UObject* Option = OptionPtr.Get())
+		{
+			if (UDatasmithCommonTessellationOptions* TessellationOptionsObject = Cast<UDatasmithCommonTessellationOptions>(Option))
+			{
+				GetCommonTessellationOptionsPtr().Reset(TessellationOptionsObject);
+			}
+		}
+	}
+}
+
+TStrongObjectPtr<UDatasmithCommonTessellationOptions>& FDatasmithCoreTechTranslator::GetCommonTessellationOptionsPtr()
+{
+	if(!CommonTessellationOptionsPtr.IsValid())
+	{
+		CommonTessellationOptionsPtr = Datasmith::MakeOptions<UDatasmithCommonTessellationOptions>();
+		check(CommonTessellationOptionsPtr.IsValid());
+		InitCommonTessellationOptions(CommonTessellationOptionsPtr->Options);
+	}
+
+	return CommonTessellationOptionsPtr;
+}
+
+
+
+
diff --git a/Engine/Plugins/Enterprise/DatasmithCADImporter/Source/DatasmithCoreTechExtension/Public/DatasmithCoreTechTranslator.h b/Engine/Plugins/Enterprise/DatasmithCADImporter/Source/DatasmithCoreTechExtension/Public/DatasmithCoreTechTranslator.h
new file mode 100644
--- /dev/null
+++ b/Engine/Plugins/Enterprise/DatasmithCADImporter/Source/DatasmithCoreTechExtension/Public/DatasmithCoreTechTranslator.h
@@ -0,0 +1,42 @@
+// Copyright 1998-2019 Epic Games, Inc. All Rights Reserved.
+
+#pragma once
+
+#include "CoreMinimal.h"
+
+#include "CADLibraryOptions.h"
+#include "DatasmithImportOptions.h"
+#include "Translators/DatasmithTranslator.h"
+#include "UObject/ObjectMacros.h"
+
+class UDatasmithCommonTessellationOptions;
+
+
+class DATASMITHCORETECHEXTENSION_API FDatasmithCoreTechTranslator : public IDatasmithTranslator
+{
+public:
+	virtual ~FDatasmithCoreTechTranslator() {}
+
+	// Begin IDatasmithTranslator overrides
+	virtual void GetSceneImportOptions(TArray<TStrongObjectPtr<UObject>>& Options) override;
+	virtual void SetSceneImportOptions(TArray<TStrongObjectPtr<UObject>>& Options) override;
+	// End IDatasmithTranslator overrides
+
+protected:
+	const FDatasmithTessellationOptions& GetCommonTessellationOptions()
+	{
+		return GetCommonTessellationOptionsPtr()->Options;
+	}
+
+	/** 
+	 * Call when the UDatasmithCommonTessellationOptions object is created. This is the unique opportunity for
+	 * child class to overwrite some values
+	 * @param TessellationOptions Reference on member of UDatasmithCommonTessellationOptions
+	 */
+	virtual void InitCommonTessellationOptions(FDatasmithTessellationOptions& TessellationOptions) {}
+
+private:
+	TStrongObjectPtr<UDatasmithCommonTessellationOptions>& GetCommonTessellationOptionsPtr();
+	TStrongObjectPtr<UDatasmithCommonTessellationOptions> CommonTessellationOptionsPtr;
+};
+
diff --git a/Engine/Plugins/Enterprise/DatasmithCADImporter/Source/DatasmithOpenNurbsTranslator/Private/DatasmithOpenNurbsTranslator.cpp b/Engine/Plugins/Enterprise/DatasmithCADImporter/Source/DatasmithOpenNurbsTranslator/Private/DatasmithOpenNurbsTranslator.cpp
--- a/Engine/Plugins/Enterprise/DatasmithCADImporter/Source/DatasmithOpenNurbsTranslator/Private/DatasmithOpenNurbsTranslator.cpp
+++ b/Engine/Plugins/Enterprise/DatasmithCADImporter/Source/DatasmithOpenNurbsTranslator/Private/DatasmithOpenNurbsTranslator.cpp
@@ -3065,7 +3065,6 @@ void FOpenNurbsTranslatorImpl::SetTessellationOptions(const FDatasmithTessellati
 FDatasmithOpenNurbsTranslator::FDatasmithOpenNurbsTranslator()
 	: Translator(nullptr)
 {
-	TessellationOptions.StitchingTechnique = EDatasmithCADStitchingTechnique::StitchingNone;
 }
 
 void FDatasmithOpenNurbsTranslator::Initialize(FDatasmithTranslatorCapabilities& OutCapabilities)
@@ -3106,7 +3105,7 @@ bool FDatasmithOpenNurbsTranslator::LoadScene(TSharedRef<IDatasmithScene> OutSce
 	IFileManager::Get().MakeDirectory(*OutputPath, true);
 	Translator->SetOutputPath(OutputPath);
 
-	Translator->SetTessellationOptions(TessellationOptions);
+	Translator->SetTessellationOptions(GetCommonTessellationOptions());
 
 	ON_BinaryFile Archive(ON::archive_mode::read3dm, FileHandle);
 
@@ -3147,7 +3146,7 @@ bool FDatasmithOpenNurbsTranslator::LoadStaticMesh(const TSharedRef<IDatasmithMe
 				CoreTechData->RawData = MoveTemp(ByteArray);
 				CoreTechData->SceneParameters.ModelCoordSys = uint8(FDatasmithUtils::EModelCoordSystem::ZUp_RightHanded);
 				CoreTechData->SceneParameters.ScaleFactor = Translator->GetScalingFactor(); // Scale is set according to the Rhino file unit
-				CoreTechData->LastTessellationOptions = TessellationOptions;
+				CoreTechData->LastTessellationOptions = GetCommonTessellationOptions();
 				OutMeshPayload.AdditionalData.Add(CoreTechData);
 			}
 		}
@@ -3160,29 +3159,14 @@ bool FDatasmithOpenNurbsTranslator::LoadStaticMesh(const TSharedRef<IDatasmithMe
 #endif
 }
 
-void FDatasmithOpenNurbsTranslator::GetSceneImportOptions(TArray<TStrongObjectPtr<UObject>>& Options)
-{
-#ifdef CAD_LIBRARY
-	TStrongObjectPtr<UDatasmithCommonTessellationOptions> TessellationOptionsPtr = Datasmith::MakeOptions<UDatasmithCommonTessellationOptions>();
-	TessellationOptionsPtr->Options.StitchingTechnique = EDatasmithCADStitchingTechnique::StitchingNone; // #ueent_todo: workaround for UE-74572
-	Options.Add(TessellationOptionsPtr);
-#endif
-}
-
 void FDatasmithOpenNurbsTranslator::SetSceneImportOptions(TArray<TStrongObjectPtr<UObject>>& Options)
 {
 #ifdef USE_OPENNURBS
-	for (const TStrongObjectPtr<UObject>& OptionPtr : Options)
+	FDatasmithCoreTechTranslator::SetSceneImportOptions(Options);
+
+	if (Translator)
 	{
-		UObject* Option = OptionPtr.Get();
-		if (UDatasmithCommonTessellationOptions* TessellationOptionsObject = Cast<UDatasmithCommonTessellationOptions>(Option))
-		{
-			TessellationOptions = TessellationOptionsObject->Options;
-			if (Translator)
-			{
-				Translator->SetTessellationOptions(TessellationOptions);
-			}
-		}
+		Translator->SetTessellationOptions( GetCommonTessellationOptions() );
 	}
 #endif
 }
diff --git a/Engine/Plugins/Enterprise/DatasmithCADImporter/Source/DatasmithOpenNurbsTranslator/Private/DatasmithOpenNurbsTranslator.h b/Engine/Plugins/Enterprise/DatasmithCADImporter/Source/DatasmithOpenNurbsTranslator/Private/DatasmithOpenNurbsTranslator.h
--- a/Engine/Plugins/Enterprise/DatasmithCADImporter/Source/DatasmithOpenNurbsTranslator/Private/DatasmithOpenNurbsTranslator.h
+++ b/Engine/Plugins/Enterprise/DatasmithCADImporter/Source/DatasmithOpenNurbsTranslator/Private/DatasmithOpenNurbsTranslator.h
@@ -2,19 +2,19 @@
 
 #pragma once
 
-#include "DatasmithImportOptions.h"
-#include "Translators/DatasmithTranslator.h"
+#include "DatasmithCoreTechTranslator.h"
 
 #include "CoreMinimal.h"
 #include "UObject/ObjectMacros.h"
 
 class FOpenNurbsTranslatorImpl;
 
-class FDatasmithOpenNurbsTranslator : public IDatasmithTranslator
+class FDatasmithOpenNurbsTranslator : public FDatasmithCoreTechTranslator
 {
 public:
 	FDatasmithOpenNurbsTranslator();
 
+	// Begin IDatasmithTranslator overrides
 	virtual FName GetFName() const override { return "DatasmithOpenNurbsTranslator"; };
 
 	virtual void Initialize(FDatasmithTranslatorCapabilities& OutCapabilities) override;
@@ -25,12 +25,20 @@ public:
 	virtual void UnloadScene() override;
 
 	virtual bool LoadStaticMesh(const TSharedRef<IDatasmithMeshElement> MeshElement, FDatasmithMeshElementPayload& OutMeshPayload) override;
+	// End IDatasmithTranslator overrides
 
-	virtual void GetSceneImportOptions(TArray<TStrongObjectPtr<UObject>>& Options) override;
+	// Begin ADatasmithCoreTechTranslator overrides
 	virtual void SetSceneImportOptions(TArray<TStrongObjectPtr<UObject>>& Options) override;
 
+protected:
+	virtual void InitCommonTessellationOptions(FDatasmithTessellationOptions& TessellationOptions) override
+	{
+		TessellationOptions.StitchingTechnique = EDatasmithCADStitchingTechnique::StitchingNone;
+	}
+	// End ADatasmithCoreTechTranslator overrides
+
 private:
+
 	TSharedPtr<FOpenNurbsTranslatorImpl> Translator;
-	FDatasmithTessellationOptions TessellationOptions;
 };
 
diff --git a/Engine/Plugins/Enterprise/DatasmithCADImporter/Source/DatasmithWireTranslator/Private/DatasmithWireTranslator.cpp b/Engine/Plugins/Enterprise/DatasmithCADImporter/Source/DatasmithWireTranslator/Private/DatasmithWireTranslator.cpp
--- a/Engine/Plugins/Enterprise/DatasmithCADImporter/Source/DatasmithWireTranslator/Private/DatasmithWireTranslator.cpp
+++ b/Engine/Plugins/Enterprise/DatasmithCADImporter/Source/DatasmithWireTranslator/Private/DatasmithWireTranslator.cpp
@@ -1441,7 +1441,7 @@ bool FDatasmithWireTranslator::LoadScene(TSharedRef<IDatasmithScene> OutScene)
 	IFileManager::Get().MakeDirectory(*OutputPath, true);
 	Translator->SetOutputPath(OutputPath);
 
-	Translator->SetTessellationOptions(TessellationOptions);
+	Translator->SetTessellationOptions(GetCommonTessellationOptions());
 
 	return Translator->Read();
 
@@ -1484,7 +1484,8 @@ bool FDatasmithWireTranslator::LoadStaticMesh(const TSharedRef<IDatasmithMeshEle
 				CoreTechData->MeshParameters.SymmetricNormal = MeshParameters.SymmetricNormal;
 				CoreTechData->MeshParameters.SymmetricOrigin = MeshParameters.SymmetricOrigin;
 
-				CoreTechData->LastTessellationOptions = TessellationOptions;
+				CoreTechData->LastTessellationOptions = GetCommonTessellationOptions();
+
 				OutMeshPayload.AdditionalData.Add(CoreTechData);
 			}
 		}
@@ -1497,29 +1498,14 @@ bool FDatasmithWireTranslator::LoadStaticMesh(const TSharedRef<IDatasmithMeshEle
 #endif
 }
 
-void FDatasmithWireTranslator::GetSceneImportOptions(TArray<TStrongObjectPtr<UObject>>& Options)
-{
-	TStrongObjectPtr<UDatasmithCommonTessellationOptions> TessellationOptionsPtr = Datasmith::MakeOptions<UDatasmithCommonTessellationOptions>();
-	TessellationOptionsPtr->Options.StitchingTechnique = EDatasmithCADStitchingTechnique::StitchingNone; // #ueent_todo: workaround for UE-74572
-	Options.Add(TessellationOptionsPtr);
-}
-
 void FDatasmithWireTranslator::SetSceneImportOptions(TArray<TStrongObjectPtr<UObject>>& Options)
 {
 #ifdef USE_OPENMODEL
-	for (const TStrongObjectPtr<UObject>& OptionPtr : Options)
+	FDatasmithCoreTechTranslator::SetSceneImportOptions(Options);
+
+	if (Translator)
 	{
-		if (UObject* Option = OptionPtr.Get())
-		{
-			if (UDatasmithCommonTessellationOptions* TessellationOptionsObject = Cast<UDatasmithCommonTessellationOptions>(Option))
-			{
-				TessellationOptions = TessellationOptionsObject->Options;
-				if (Translator)
-				{
-					Translator->SetTessellationOptions(TessellationOptions);
-				}
-			}
-		}
+		Translator->SetTessellationOptions( GetCommonTessellationOptions() );
 	}
 #endif
 }
diff --git a/Engine/Plugins/Enterprise/DatasmithCADImporter/Source/DatasmithWireTranslator/Private/DatasmithWireTranslator.h b/Engine/Plugins/Enterprise/DatasmithCADImporter/Source/DatasmithWireTranslator/Private/DatasmithWireTranslator.h
--- a/Engine/Plugins/Enterprise/DatasmithCADImporter/Source/DatasmithWireTranslator/Private/DatasmithWireTranslator.h
+++ b/Engine/Plugins/Enterprise/DatasmithCADImporter/Source/DatasmithWireTranslator/Private/DatasmithWireTranslator.h
@@ -5,19 +5,19 @@
 #include "CoreMinimal.h"
 
 #include "CADLibraryOptions.h"
-#include "DatasmithImportOptions.h"
-#include "Translators/DatasmithTranslator.h"
+#include "DatasmithCoreTechTranslator.h"
 #include "UObject/ObjectMacros.h"
 
 class FWireTranslatorImpl;
 class AlDagNode;
 
 
-class FDatasmithWireTranslator : public IDatasmithTranslator
+class FDatasmithWireTranslator : public FDatasmithCoreTechTranslator
 {
 public:
 	FDatasmithWireTranslator();
 
+	// Begin IDatasmithTranslator overrides
 	virtual FName GetFName() const override { return "DatasmithWireTranslator"; };
 
 	virtual void Initialize(FDatasmithTranslatorCapabilities& OutCapabilities) override;
@@ -28,13 +28,20 @@ public:
 	virtual void UnloadScene() override;
 
 	virtual bool LoadStaticMesh(const TSharedRef<IDatasmithMeshElement> MeshElement, FDatasmithMeshElementPayload& OutMeshPayload) override;
+	// End IDatasmithTranslator overrides
 
-	virtual void GetSceneImportOptions(TArray<TStrongObjectPtr<UObject>>& Options) override;
+	// Begin ADatasmithCoreTechTranslator overrides
 	virtual void SetSceneImportOptions(TArray<TStrongObjectPtr<UObject>>& Options) override;
 
+protected:
+	virtual void InitCommonTessellationOptions(FDatasmithTessellationOptions& TessellationOptions) override
+	{
+		TessellationOptions.StitchingTechnique = EDatasmithCADStitchingTechnique::StitchingNone;
+	}
+	// End ADatasmithCoreTechTranslator overrides
+
 private:
 	TSharedPtr<FWireTranslatorImpl> Translator;
-	FDatasmithTessellationOptions TessellationOptions;
 
 	const float MetricUnit = 0.001; // As CT don`t take care of unit input data, metric unit is set with default value
 	const float ScaleFactor = 1; // Wire internal unit is cm so no scale needed
diff --git a/Engine/Plugins/Enterprise/DatasmithContent/Source/DatasmithContent/Public/DatasmithImportOptions.h b/Engine/Plugins/Enterprise/DatasmithContent/Source/DatasmithContent/Public/DatasmithImportOptions.h
--- a/Engine/Plugins/Enterprise/DatasmithContent/Source/DatasmithContent/Public/DatasmithImportOptions.h
+++ b/Engine/Plugins/Enterprise/DatasmithContent/Source/DatasmithContent/Public/DatasmithImportOptions.h
@@ -212,13 +212,12 @@ struct DATASMITHCONTENT_API FDatasmithTessellationOptions
 {
 	GENERATED_BODY()
 
-	FDatasmithTessellationOptions()
-		: ChordTolerance(0.2f)
-		, MaxEdgeLength(0.0f)
-		, NormalTolerance(20.0f)
-		, StitchingTechnique(EDatasmithCADStitchingTechnique::StitchingSew)
+	FDatasmithTessellationOptions(float InChordTolerance = 0.2f, float InMaxEdgeLength = 0.0f, float InNormalTolerance = 20.0f, EDatasmithCADStitchingTechnique InStitchingTechnique = EDatasmithCADStitchingTechnique::StitchingSew)
+		: ChordTolerance(InChordTolerance)
+		, MaxEdgeLength(InMaxEdgeLength)
+		, NormalTolerance(InNormalTolerance)
+		, StitchingTechnique(InStitchingTechnique)
 	{
-
 	}
 
 	/**
diff --git a/Engine/Plugins/Enterprise/DatasmithImporter/Config/BaseDatasmithImporter.ini b/Engine/Plugins/Enterprise/DatasmithImporter/Config/BaseDatasmithImporter.ini
--- a/Engine/Plugins/Enterprise/DatasmithImporter/Config/BaseDatasmithImporter.ini
+++ b/Engine/Plugins/Enterprise/DatasmithImporter/Config/BaseDatasmithImporter.ini
@@ -1,3 +1,17 @@
 [CoreRedirects]
 
 +ClassRedirects=(OldName="/Script/DatasmithImporter.DatasmithBlutilityActor",NewName="/Script/Blutility.PlacedEditorUtilityBase")
+
+[FileProducerTessellationOptions]
+ChordTolerance=0.3
+MaxEdgeLength=0.0
+NormalTolerance=30.0
+StitchingTechnique=StitchingSew
+
+[FileProducerImportOptions]
+SceneHandling=CurrentLevel
+IncludeGeometry=True
+IncludeMaterial=True
+IncludeLight=True
+IncludeCamera=True
+IncludeAnimation=True
diff --git a/Engine/Plugins/Enterprise/DatasmithImporter/Source/DatasmithImporter/Private/DatasmithFileProducer.cpp b/Engine/Plugins/Enterprise/DatasmithImporter/Source/DatasmithImporter/Private/DatasmithFileProducer.cpp
--- a/Engine/Plugins/Enterprise/DatasmithImporter/Source/DatasmithImporter/Private/DatasmithFileProducer.cpp
+++ b/Engine/Plugins/Enterprise/DatasmithImporter/Source/DatasmithImporter/Private/DatasmithFileProducer.cpp
@@ -16,6 +16,7 @@
 #include "IDatasmithSceneElements.h"
 #include "Translators/DatasmithTranslatorManager.h"
 #include "Utility/DatasmithImporterUtils.h"
+
 #include "Async/ParallelFor.h"
 #include "Algo/Count.h"
 #include "AssetRegistryModule.h"
@@ -37,6 +38,7 @@
 #include "Materials/Material.h"
 #include "Materials/MaterialInterface.h"
 #include "Materials/MaterialInstance.h"
+#include "Misc/ConfigCacheIni.h"
 #include "ObjectTools.h"
 #include "PropertyHandle.h"
 #include "Widgets/Images/SImage.h"
@@ -60,6 +62,9 @@ const TCHAR* ExtensionPrefix = TEXT( "*." );
 
 TSet< FString > UDatasmithDirProducer::SupportedFormats;
 
+FDatasmithTessellationOptions UDatasmithFileProducer::DefaultTessellationOptions( 0.3f, 0.0f, 30.0f, EDatasmithCADStitchingTechnique::StitchingSew );
+FDatasmithImportBaseOptions UDatasmithFileProducer::DefaultImportOptions;
+
 namespace FDatasmithFileProducerUtils
 {
 	/** Delete all the assets stored under the specified path */
@@ -106,10 +111,35 @@ bool UDatasmithFileProducer::Initialize()
 		return false;
 	}
 
+	// Set all import options to defaults for Dataprep
+	TSharedPtr<IDatasmithTranslator> TranslatorPtr = TranslatableSourcePtr->GetTranslator();
+	if(IDatasmithTranslator* Translator = TranslatorPtr.Get())
+	{
+		TArray< TStrongObjectPtr<UObject> > Options;
+		Translator->GetSceneImportOptions( Options );
+
+		bool bUpdateOptions = false;
+		for(TStrongObjectPtr<UObject>& ObjectPtr : Options)
+		{
+			if(UDatasmithCommonTessellationOptions* TessellationOption = Cast<UDatasmithCommonTessellationOptions>(ObjectPtr.Get()))
+			{
+				bUpdateOptions = true;
+				TessellationOption->Options = DefaultTessellationOptions;
+			}
+		}
+
+		if(bUpdateOptions == true)
+		{
+			Translator->SetSceneImportOptions( Options );
+		}
+	}
+
 	// Create and initialize context
 	ImportContextPtr = MakeUnique< FDatasmithImportContext >( Source.GetSourceFile(), false, TEXT("DatasmithFileProducer"), LOCTEXT("DatasmithFileProducerDescription", "Datasmith File Producer"), TranslatableSourcePtr->GetTranslator() );
 
-	ImportContextPtr->Options->BaseOptions.SceneHandling = EDatasmithImportScene::CurrentLevel;
+	// Set import options to default
+	ImportContextPtr->Options->BaseOptions = DefaultImportOptions;
+
 	ImportContextPtr->SceneAsset = DatasmithScenePtr.Get();
 	ImportContextPtr->ActorsContext.ImportWorld = Context.WorldPtr.Get();
 
@@ -133,6 +163,7 @@ bool UDatasmithFileProducer::Initialize()
 		return false;
 	}
 
+
 	// Fill up scene element with content of input file
 	if (!TranslatableSourcePtr->Translate( SceneElement ))
 	{
@@ -1247,4 +1278,57 @@ void FDatasmithFileProducerUtils::DeletePackagePath( const FString& PathToDelete
 	}
 }
 
+void UDatasmithFileProducer::LoadDefaultSettings()
+{
+	// Read default settings, tessellation and import, for Datasmith file producer
+	const FString DatasmithImporterIni = FString::Printf(TEXT("%s%s/%s.ini"), *FPaths::GeneratedConfigDir(), ANSI_TO_TCHAR(FPlatformProperties::PlatformName()), TEXT("DatasmithImporter") );
+
+	const TCHAR* TessellationSectionName = TEXT("FileProducerTessellationOptions");
+	if(GConfig->DoesSectionExist( TessellationSectionName, DatasmithImporterIni ))
+	{
+
+		GConfig->GetFloat( TessellationSectionName, TEXT("ChordTolerance"), DefaultTessellationOptions.ChordTolerance, DatasmithImporterIni);
+		GConfig->GetFloat( TessellationSectionName, TEXT("MaxEdgeLength"), DefaultTessellationOptions.MaxEdgeLength, DatasmithImporterIni);
+		GConfig->GetFloat( TessellationSectionName, TEXT("NormalTolerance"), DefaultTessellationOptions.NormalTolerance, DatasmithImporterIni);
+
+		FString StitchingTechnique = GConfig->GetStr( TessellationSectionName, TEXT("StitchingTechnique"), DatasmithImporterIni);
+		if(StitchingTechnique == TEXT("StitchingHeal"))
+		{
+			DefaultTessellationOptions.StitchingTechnique =  EDatasmithCADStitchingTechnique::StitchingHeal;
+		}
+		else if(StitchingTechnique == TEXT("StitchingSew"))
+		{
+			DefaultTessellationOptions.StitchingTechnique =  EDatasmithCADStitchingTechnique::StitchingSew;
+		}
+		else
+		{
+			DefaultTessellationOptions.StitchingTechnique =  EDatasmithCADStitchingTechnique::StitchingNone;
+		}
+	}
+
+	const TCHAR* ImportSectionName = TEXT("FileProducerImportOptions");
+	if(GConfig->DoesSectionExist( ImportSectionName, DatasmithImporterIni ))
+	{
+		GConfig->GetBool( ImportSectionName, TEXT("IncludeGeometry"), DefaultImportOptions.bIncludeGeometry, DatasmithImporterIni);
+		GConfig->GetBool( ImportSectionName, TEXT("IncludeMaterial"), DefaultImportOptions.bIncludeMaterial, DatasmithImporterIni);
+		GConfig->GetBool( ImportSectionName, TEXT("IncludeLight"), DefaultImportOptions.bIncludeLight, DatasmithImporterIni);
+		GConfig->GetBool( ImportSectionName, TEXT("IncludeCamera"), DefaultImportOptions.bIncludeCamera, DatasmithImporterIni);
+		GConfig->GetBool( ImportSectionName, TEXT("IncludeAnimation"), DefaultImportOptions.bIncludeAnimation, DatasmithImporterIni);
+
+		FString SceneHandling = GConfig->GetStr( ImportSectionName, TEXT("SceneHandling"), DatasmithImporterIni);
+		if(SceneHandling == TEXT("NewLevel"))
+		{
+			DefaultImportOptions.SceneHandling =  EDatasmithImportScene::NewLevel;
+		}
+		else if(SceneHandling == TEXT("AssetsOnly"))
+		{
+			DefaultImportOptions.SceneHandling =  EDatasmithImportScene::AssetsOnly;
+		}
+		else
+		{
+			DefaultImportOptions.SceneHandling =  EDatasmithImportScene::CurrentLevel;
+		}
+	}
+}
+
 #undef LOCTEXT_NAMESPACE
diff --git a/Engine/Plugins/Enterprise/DatasmithImporter/Source/DatasmithImporter/Private/DatasmithFileProducer.h b/Engine/Plugins/Enterprise/DatasmithImporter/Source/DatasmithImporter/Private/DatasmithFileProducer.h
--- a/Engine/Plugins/Enterprise/DatasmithImporter/Source/DatasmithImporter/Private/DatasmithFileProducer.h
+++ b/Engine/Plugins/Enterprise/DatasmithImporter/Source/DatasmithImporter/Private/DatasmithFileProducer.h
@@ -8,6 +8,7 @@
 
 #include "DatasmithContentEditorModule.h"
 #include "DatasmithImportContext.h"
+#include "DatasmithImportOptions.h"
 #include "DatasmithScene.h"
 #include "Translators/DatasmithTranslatableSource.h"
 
@@ -31,6 +32,9 @@ public:
 
 	const FString& GetFilePath() const { return FilePath; }
 
+	/** Load default settings for file producer in DatasmithImporter.ini */
+	static void LoadDefaultSettings();
+
 	// Begin UDataprepContentProducer overrides
 	virtual const FText& GetLabel() const override;
 	virtual const FText& GetDescription() const override;
@@ -62,6 +66,9 @@ private:
 
 	TArray< TWeakObjectPtr< UObject > > Assets;
 
+	static FDatasmithTessellationOptions DefaultTessellationOptions;
+	static FDatasmithImportBaseOptions DefaultImportOptions;
+
 	friend class SDatasmithFileProducerFileProperty;
 };
 
diff --git a/Engine/Plugins/Enterprise/DatasmithImporter/Source/DatasmithImporter/Private/DatasmithImporterModule.cpp b/Engine/Plugins/Enterprise/DatasmithImporter/Source/DatasmithImporter/Private/DatasmithImporterModule.cpp
--- a/Engine/Plugins/Enterprise/DatasmithImporter/Source/DatasmithImporter/Private/DatasmithImporterModule.cpp
+++ b/Engine/Plugins/Enterprise/DatasmithImporter/Source/DatasmithImporter/Private/DatasmithImporterModule.cpp
@@ -61,6 +61,8 @@ public:
 		FDatasmithMasterMaterialManager::Get().RegisterSelector( TEXT("SketchUp"), MakeShared< FDatasmithSketchUpMaterialSelector >() );
 		FDatasmithMasterMaterialManager::Get().RegisterSelector( TEXT("CityEngine"), MakeShared< FDatasmithCityEngineMaterialSelector >() );
 
+		UDatasmithFileProducer::LoadDefaultSettings();
+
 		// Disable any UI feature if running in command mode
 		if (!IsRunningCommandlet())
 		{
