Integrating CL# 9452284 from Dev-Enterprise
Datasmith: Fixing empty mesh section when importing IFC files

diff --git a/Engine/Plugins/Enterprise/DatasmithIFCImporter/Source/DatasmithIFCTranslator/Private/IFC/IFCStaticMeshFactory.cpp b/Engine/Plugins/Enterprise/DatasmithIFCImporter/Source/DatasmithIFCTranslator/Private/IFC/IFCStaticMeshFactory.cpp
--- a/Engine/Plugins/Enterprise/DatasmithIFCImporter/Source/DatasmithIFCTranslator/Private/IFC/IFCStaticMeshFactory.cpp
+++ b/Engine/Plugins/Enterprise/DatasmithIFCImporter/Source/DatasmithIFCTranslator/Private/IFC/IFCStaticMeshFactory.cpp
@@ -3,6 +3,7 @@
 #include "IFCStaticMeshFactory.h"
 
 #include "AssetRegistryModule.h"
+#include "DatasmithMeshHelper.h"
 #include "Engine/StaticMesh.h"
 #include "StaticMeshAttributes.h"
 #include "MeshDescription.h"
@@ -119,6 +120,8 @@ namespace IFC
 
 			MeshDescription->CreatePolygon(MaterialIndexToPolygonGroupID[IFCPolygon->MaterialIndex], VertexInstanceIDs);
 		}
+
+		DatasmithMeshHelper::RemoveEmptyPolygonGroups(*MeshDescription);
 	}
 
 	float FStaticMeshFactory::GetUniformScale() const
diff --git a/Engine/Plugins/Enterprise/DatasmithImporter/Source/DatasmithImporter/Private/DatasmithMeshHelper.cpp b/Engine/Plugins/Enterprise/DatasmithImporter/Source/DatasmithImporter/Private/DatasmithMeshHelper.cpp
--- a/Engine/Plugins/Enterprise/DatasmithImporter/Source/DatasmithImporter/Private/DatasmithMeshHelper.cpp
+++ b/Engine/Plugins/Enterprise/DatasmithImporter/Source/DatasmithImporter/Private/DatasmithMeshHelper.cpp
@@ -60,6 +60,25 @@ namespace DatasmithMeshHelper
 		return NormalLengthSquare < SMALL_NUMBER;
 	}
 
+	void RemoveEmptyPolygonGroups(FMeshDescription& Mesh)
+	{
+		bool bRemovedSection = false;
+		for (const FPolygonGroupID& PolygonGroupID : Mesh.PolygonGroups().GetElementIDs())
+		{
+			if (Mesh.GetNumPolygonGroupPolygons(PolygonGroupID) == 0)
+			{
+				Mesh.DeletePolygonGroup(PolygonGroupID);
+				bRemovedSection = true;
+			}
+		}
+
+		if (bRemovedSection)
+		{
+			FElementIDRemappings Remappings;
+			Mesh.Compact(Remappings);
+		}
+	}
+
 	FMeshDescription* InitMeshDescription(UStaticMesh* StaticMesh, int32 LodIndex)
 	{
 		if (!ensure(StaticMesh))
diff --git a/Engine/Plugins/Enterprise/DatasmithImporter/Source/DatasmithImporter/Public/DatasmithMeshHelper.h b/Engine/Plugins/Enterprise/DatasmithImporter/Source/DatasmithImporter/Public/DatasmithMeshHelper.h
--- a/Engine/Plugins/Enterprise/DatasmithImporter/Source/DatasmithImporter/Public/DatasmithMeshHelper.h
+++ b/Engine/Plugins/Enterprise/DatasmithImporter/Source/DatasmithImporter/Public/DatasmithMeshHelper.h
@@ -60,6 +60,13 @@ namespace DatasmithMeshHelper
 	 */
 	DATASMITHIMPORTER_API bool IsTriangleDegenerated(const FMeshDescription& Mesh, const FMeshTriangle& MeshTriangle);
 
+	/**
+	 * Remove any empty polygon groups in the mesh description and compact it if needed.
+	 *
+	 * @param Mesh		    The source mesh
+	 */
+	DATASMITHIMPORTER_API void RemoveEmptyPolygonGroups(FMeshDescription& Mesh);
+
 	/**
 	 * Write mesh data on a UStaticMesh.
 	 * RawMesh will be converted to MeshDescription beforehand.
