Merging CLs #9442540 and #9455962 from //UE4/Dev-Enterprise

9442540:
Fixed a case typo in an include

9455962:
Implemented undo/redo for UI related to content producers and consumer
Fixed issue with content folder entered thru the keyboard

diff --git a/Engine/Plugins/Enterprise/DatasmithImporter/Source/DatasmithImporter/Private/DatasmithConsumer.cpp b/Engine/Plugins/Enterprise/DatasmithImporter/Source/DatasmithImporter/Private/DatasmithConsumer.cpp
--- a/Engine/Plugins/Enterprise/DatasmithImporter/Source/DatasmithImporter/Private/DatasmithConsumer.cpp
+++ b/Engine/Plugins/Enterprise/DatasmithImporter/Source/DatasmithImporter/Private/DatasmithConsumer.cpp
@@ -309,40 +309,53 @@ ULevel * UDatasmithConsumer::FindLevel( const FString& InLevelName )
 
 bool UDatasmithConsumer::SetLevelName( const FString & InLevelName, FText& OutReason )
 {
-	// Update current level
-	if( !InLevelName.IsEmpty() && InLevelName.Compare( TEXT("current"), ESearchCase::IgnoreCase ) != 0 )
+	FString NewLevelName = InLevelName;
+
+	bool bValidLevelName = false;
+	OutReason = FText();
+
+	// Check if a new level can be used with the new name and current limitations
+	if( !NewLevelName.IsEmpty() && NewLevelName.Compare( TEXT("current"), ESearchCase::IgnoreCase ) != 0 )
 	{
-		// #ueent_todo: What about sub-level of sub-level?
+		// Sub-level of sub-level is not supported yet
+		// #ueent_todo: sub-level of sub-level
 		if( InLevelName.Contains( TEXT("/") ) || InLevelName.Contains( TEXT("\\") ))
 		{
 			OutReason = LOCTEXT( "DatasmithConsumer_SubLevel", "Sub-level of sub-levels is not supported yet" );
-			return false;
 		}
-
 		// Try to see if there is any issue to eventually create this level, i.e. name collision
-		if( FindLevel( InLevelName ) == nullptr )
+		else if( FindLevel( InLevelName ) == nullptr )
 		{
 			FSoftObjectPath LevelObjectPath( FPaths::Combine( TargetContentFolder, InLevelName ) );
 
 			if( StaticFindObject( nullptr, ANY_PACKAGE, *LevelObjectPath.ToString(), true) )
 			{
 				OutReason = LOCTEXT( "DatasmithConsumer_LevelExists", "A object with that name already exists. Please choose another name." );
-				return false;
 			}
 
 			// #ueent_todo: Check if persistent level is locked, etc
 		}
 
-		LevelName = InLevelName;
+		// Good to go if no error documented
+		bValidLevelName = OutReason.IsEmpty();
 	}
-	else
+	// New name of level is empty or keyword 'current' used
+	else if( !LevelName.IsEmpty() )
 	{
-		LevelName = TEXT("");
+		NewLevelName = TEXT("");
+		bValidLevelName = true;
 	}
 
-	OnChanged.Broadcast();
+	if(bValidLevelName)
+	{
+		Modify();
 
-	return true;
+		LevelName = NewLevelName;
+
+		OnChanged.Broadcast();
+	}
+
+	return bValidLevelName;
 }
 
 void UDatasmithConsumer::MoveAssets()
diff --git a/Engine/Plugins/Enterprise/DatasmithImporter/Source/DatasmithImporter/Private/DatasmithFileProducer.cpp b/Engine/Plugins/Enterprise/DatasmithImporter/Source/DatasmithImporter/Private/DatasmithFileProducer.cpp
--- a/Engine/Plugins/Enterprise/DatasmithImporter/Source/DatasmithImporter/Private/DatasmithFileProducer.cpp
+++ b/Engine/Plugins/Enterprise/DatasmithImporter/Source/DatasmithImporter/Private/DatasmithFileProducer.cpp
@@ -9,6 +9,7 @@
 #include "DatasmithAssetUserData.h"
 #include "DatasmithContentBlueprintLibrary.h"
 #include "DatasmithImportContext.h"
+#include "DatasmithImportOptions.h"
 #include "DatasmithImporter.h"
 #include "DatasmithScene.h"
 #include "DatasmithSceneActor.h"
@@ -41,6 +42,7 @@
 #include "Misc/ConfigCacheIni.h"
 #include "ObjectTools.h"
 #include "PropertyHandle.h"
+#include "ScopedTransaction.h"
 #include "Widgets/Images/SImage.h"
 #include "Widgets/Input/SButton.h"
 #include "Widgets/Input/SEditableTextBox.h"
@@ -78,6 +80,12 @@ namespace FDatasmithFileProducerUtils
 			DeletePackagePath( PathToDelete );
 		}
 	}
+
+	/** Display OS browser, i.e. Windows explorer, to let user select a file */
+	FString SelectFileToImport();
+
+	/** Display OS browser, i.e. Windows explorer, to let user select a directory */
+	FString SelectDirectory();
 }
 
 bool UDatasmithFileProducer::Initialize()
@@ -163,7 +171,6 @@ bool UDatasmithFileProducer::Initialize()
 		return false;
 	}
 
-
 	// Fill up scene element with content of input file
 	if (!TranslatableSourcePtr->Translate( SceneElement ))
 	{
@@ -542,6 +549,8 @@ FString UDatasmithFileProducer::GetNamespace() const
 
 void UDatasmithFileProducer::SetFilename( const FString& InFilename )
 {
+	Modify();
+
 	FilePath = FPaths::ConvertRelativePathToFull( InFilename );
 
 	// Rename producer to name of file
@@ -576,6 +585,24 @@ bool UDatasmithFileProducer::Supersede(const UDataprepContentProducer* OtherProd
 		FilePath == OtherFileProducer->FilePath;
 }
 
+void UDatasmithFileProducer::PostEditUndo()
+{
+	UDataprepContentProducer::PostEditUndo();
+
+	OnChanged.Broadcast( this );
+}
+
+void UDatasmithFileProducer::PostInitProperties()
+{
+	UDataprepContentProducer::PostInitProperties();
+
+	// Set FilePath when creating a new producer
+	if( !HasAnyFlags( RF_ClassDefaultObject | RF_WasLoaded | RF_Transient ) )
+	{
+		FilePath = FDatasmithFileProducerUtils::SelectFileToImport();
+	}
+}
+
 UDatasmithDirProducer::UDatasmithDirProducer()
 	: ExtensionString( TEXT("*.*") )
 	, bRecursive(true)
@@ -630,7 +657,7 @@ bool UDatasmithDirProducer::Initialize()
 
 	Context.SetRootPackage( TransientPackage );
 
-	FileProducer = TStrongObjectPtr< UDatasmithFileProducer >( NewObject< UDatasmithFileProducer >() );
+	FileProducer = TStrongObjectPtr< UDatasmithFileProducer >( NewObject< UDatasmithFileProducer >( GetTransientPackage(), NAME_None, RF_Transient ) );
 
 	return true;
 }
@@ -647,7 +674,7 @@ bool UDatasmithDirProducer::Execute(TArray< TWeakObjectPtr< UObject > >& OutAsse
 
 	for( const FString& FileName : FilesToProcess )
 	{
-		FileProducer->SetFilename( FileName );
+		FileProducer->FilePath =  FPaths::ConvertRelativePathToFull( FileName );
 
 		Task.ReportNextStep( FText::Format( LOCTEXT( "DatasmithFileProducer_LoadingFile", "Loading {0} ..."), FText::FromString( FileName ) ) );
 
@@ -694,8 +721,20 @@ void UDatasmithDirProducer::Serialize( FArchive& Ar )
 	}
 }
 
+void UDatasmithDirProducer::PostInitProperties()
+{
+	UDataprepContentProducer::PostInitProperties();
+
+	if( !HasAnyFlags(RF_ClassDefaultObject | RF_NeedLoad) )
+	{
+		FolderPath = FDatasmithFileProducerUtils::SelectDirectory();
+	}
+}
+
 void UDatasmithDirProducer::SetFolderName( const FString& InFolderName )
 {
+	Modify();
+
 	FolderPath = FPaths::ConvertRelativePathToFull( InFolderName );
 
 	FString BaseName = FPaths::IsDrive( InFolderName ) ? TEXT("RootDir") : FPaths::GetBaseFilename( InFolderName );
@@ -867,7 +906,6 @@ TSet<FString> UDatasmithDirProducer::GetSetOfFiles() const
 	return FoundFiles;
 }
 
-
 class SDatasmithFileProducerFileProperty : public SCompoundWidget
 {
 public:
@@ -877,10 +915,6 @@ public:
 	SLATE_ARGUMENT(UDatasmithFileProducer*, Producer)
 	SLATE_END_ARGS()
 
-	SDatasmithFileProducerFileProperty()
-		: bIsFirstPaint(true)
-	{}
-
 public:
 	void Construct(const FArguments& InArgs)
 	{
@@ -916,29 +950,6 @@ public:
 		];
 	}
 
-	virtual int32 OnPaint(const FPaintArgs& Args, const FGeometry& AllottedGeometry, const FSlateRect& MyCullingRect, FSlateWindowElementList& OutDrawElements, int32 LayerId, const FWidgetStyle& InWidgetStyle, bool bParentEnabled) const override
-	{
-		int32 Ret = SCompoundWidget::OnPaint(Args, AllottedGeometry, MyCullingRect, OutDrawElements, LayerId, InWidgetStyle, bParentEnabled );
-
-		if( bIsFirstPaint )
-		{
-			bIsFirstPaint = false;
-			if( ProducerPtr->FilePath.IsEmpty() )
-			{
-				// Cache current hit test grid since it will be wiped out by the call to OnChangePathClicked
-				// which trigger a full redraw of all windows using the same hit test grid
-				FHittestGrid CachedHittestGrid = Args.GetHittestGrid();
-
-				OnChangePathClicked();
-
-				// Restore hit test grid
-				Args.GetHittestGrid() = CachedHittestGrid;
-			}
-		}
-
-		return Ret;
-	}
-
 private:
 	FReply OnChangePathClicked() const
 	{
@@ -947,64 +958,12 @@ private:
 			return FReply::Unhandled();
 		}
 
-		const TArray<FString>& Formats = FDatasmithTranslatorManager::Get().GetSupportedFormats();
-
-		FString FileTypes;
-		FString AllExtensions;
-
-		for( const FString& Format : Formats )
-		{
-			TArray<FString> FormatComponents;
-			Format.ParseIntoArray( FormatComponents, TEXT( ";" ), false );
-
-			for ( int32 ComponentIndex = 0; ComponentIndex < FormatComponents.Num(); ComponentIndex += 2 )
-			{
-				check( FormatComponents.IsValidIndex( ComponentIndex + 1 ) );
-				const FString& Extension = FormatComponents[ComponentIndex];
-				const FString& Description = FormatComponents[ComponentIndex + 1];
-
-				if ( !AllExtensions.IsEmpty() )
-				{
-					AllExtensions.AppendChar( TEXT( ';' ) );
-				}
-				AllExtensions.Append( TEXT( "*." ) );
-				AllExtensions.Append( Extension );
-
-				if ( !FileTypes.IsEmpty() )
-				{
-					FileTypes.AppendChar( TEXT( '|' ) );
-				}
-
-				FileTypes.Append( FString::Printf( TEXT( "%s (*.%s)|*.%s" ), *Description, *Extension, *Extension ) );
-			}
-		}
-
-		FString SupportedExtensions( FString::Printf( TEXT( "All Files (%s)|%s|%s" ), *AllExtensions, *AllExtensions, *FileTypes ) );
-
-		TArray<FString> OpenedFiles;
-		FString DefaultLocation( FEditorDirectories::Get().GetLastDirectory( ELastDirectory::GENERIC_IMPORT ) );
-		IDesktopPlatform* DesktopPlatform = FDesktopPlatformModule::Get();
-
-		bool bOpened = false;
-		if ( DesktopPlatform )
-		{
-			bOpened = DesktopPlatform->OpenFileDialog(
-				FSlateApplication::Get().FindBestParentWindowHandleForDialogs( nullptr ),
-				LOCTEXT( "FileDialogTitle", "Import Datasmith" ).ToString(),
-				DefaultLocation,
-				TEXT( "" ),
-				SupportedExtensions,
-				EFileDialogFlags::None,
-				OpenedFiles
-			);
-		}
-
-		if ( bOpened && OpenedFiles.Num() > 0 )
+		FString SelectedFile = FDatasmithFileProducerUtils::SelectFileToImport();
+		if(!SelectedFile.IsEmpty())
 		{
-			const FString& OpenedFile = OpenedFiles[0];
-			FEditorDirectories::Get().SetLastDirectory( ELastDirectory::GENERIC_IMPORT, OpenedFile );
+			const FScopedTransaction Transaction( LOCTEXT("Producer_SetFilename", "Set Filename") );
 
-			ProducerPtr->SetFilename( OpenedFile );
+			ProducerPtr->SetFilename( SelectedFile );
 			FileName->SetText( GetFilenameText() );
 		}
 
@@ -1019,9 +978,6 @@ private:
 private:
 	TWeakObjectPtr< UDatasmithFileProducer > ProducerPtr;
 	TSharedPtr< SEditableText > FileName;
-
-	// Boolean used to detect first paint and pop up file dialog if file input is empty
-	mutable bool bIsFirstPaint;
 };
 
 void FDatasmithFileProducerDetails::CustomizeDetails(IDetailLayoutBuilder& DetailBuilder)
@@ -1070,10 +1026,6 @@ public:
 	SLATE_ARGUMENT( UDatasmithDirProducer*, Producer )
 	SLATE_END_ARGS()
 
-		SDatasmithDirProducerFolderProperty()
-		: bIsFirstPaint(true)
-	{}
-
 public:
 	void Construct(const FArguments& InArgs)
 	{
@@ -1109,29 +1061,6 @@ public:
 		];
 	}
 
-	virtual int32 OnPaint(const FPaintArgs& Args, const FGeometry& AllottedGeometry, const FSlateRect& MyCullingRect, FSlateWindowElementList& OutDrawElements, int32 LayerId, const FWidgetStyle& InWidgetStyle, bool bParentEnabled) const override
-	{
-		int32 Ret = SCompoundWidget::OnPaint(Args, AllottedGeometry, MyCullingRect, OutDrawElements, LayerId, InWidgetStyle, bParentEnabled );
-
-		if( bIsFirstPaint )
-		{
-			bIsFirstPaint = false;
-			if( ProducerPtr->FolderPath.IsEmpty() )
-			{
-				// Cache current hit test grid since it will be wiped out by the call to OnChangePathClicked
-				// which trigger a full redraw of all windows using the same hit test grid
-				FHittestGrid CachedHittestGrid = Args.GetHittestGrid();
-
-				OnChangePathClicked();
-
-				// Restore hit test grid
-				Args.GetHittestGrid() = CachedHittestGrid;
-			}
-		}
-
-		return Ret;
-	}
-
 private:
 	FReply OnChangePathClicked() const
 	{
@@ -1140,29 +1069,13 @@ private:
 			return FReply::Unhandled();
 		}
 
-		if( IDesktopPlatform* DesktopPlatform = FDesktopPlatformModule::Get() )
+		const FString SelectedFolder = FDatasmithFileProducerUtils::SelectDirectory();
+		if( !SelectedFolder.IsEmpty() )
 		{
-			FString DestinationFolder;
-			const void* ParentWindowHandle = FSlateApplication::Get().FindBestParentWindowHandleForDialogs(nullptr);
-			const FString Title = LOCTEXT("DatasmithDirProducerFolderTitle", "Choose a folder").ToString();
-			const FString DefaultLocation( FEditorDirectories::Get().GetLastDirectory( ELastDirectory::GENERIC_IMPORT ) );
-
-			// #ueent_todo: What if the user selects a drive? Shouldn't we ask if this is really what he/she wants?
-			const bool bFolderSelected = DesktopPlatform->OpenDirectoryDialog(
-				ParentWindowHandle,
-				Title,
-				DefaultLocation,
-				DestinationFolder
-			);
-
-			if( bFolderSelected )
-			{
-				FEditorDirectories::Get().SetLastDirectory(ELastDirectory::GENERIC_EXPORT, DestinationFolder);
-				FPaths::NormalizeFilename(DestinationFolder);
+			const FScopedTransaction Transaction( LOCTEXT("Producer_SetFolderName", "Set Folder Name") );
+			ProducerPtr->SetFolderName( SelectedFolder );
 
-				ProducerPtr->SetFolderName( DestinationFolder );
-				FolderName->SetText( GetFilenameText() );
-			}
+			FolderName->SetText( GetFilenameText() );
 		}
 
 		return FReply::Handled();
@@ -1176,9 +1089,6 @@ private:
 private:
 	TWeakObjectPtr< UDatasmithDirProducer > ProducerPtr;
 	TSharedPtr< SEditableText > FolderName;
-
-	// Boolean used to detect first paint and pop up folder dialog if folder input is empty
-	mutable bool bIsFirstPaint;
 };
 
 void FDatasmithDirProducerDetails::CustomizeDetails( IDetailLayoutBuilder& DetailBuilder )
@@ -1339,4 +1249,97 @@ void UDatasmithFileProducer::LoadDefaultSettings()
 	}
 }
 
+FString FDatasmithFileProducerUtils::SelectFileToImport()
+{
+	const TArray<FString>& Formats = FDatasmithTranslatorManager::Get().GetSupportedFormats();
+
+	FString FileTypes;
+	FString AllExtensions;
+
+	for( const FString& Format : Formats )
+	{
+		TArray<FString> FormatComponents;
+		Format.ParseIntoArray( FormatComponents, TEXT( ";" ), false );
+
+		for ( int32 ComponentIndex = 0; ComponentIndex < FormatComponents.Num(); ComponentIndex += 2 )
+		{
+			check( FormatComponents.IsValidIndex( ComponentIndex + 1 ) );
+			const FString& Extension = FormatComponents[ComponentIndex];
+			const FString& Description = FormatComponents[ComponentIndex + 1];
+
+			if ( !AllExtensions.IsEmpty() )
+			{
+				AllExtensions.AppendChar( TEXT( ';' ) );
+			}
+			AllExtensions.Append( TEXT( "*." ) );
+			AllExtensions.Append( Extension );
+
+			if ( !FileTypes.IsEmpty() )
+			{
+				FileTypes.AppendChar( TEXT( '|' ) );
+			}
+
+			FileTypes.Append( FString::Printf( TEXT( "%s (*.%s)|*.%s" ), *Description, *Extension, *Extension ) );
+		}
+	}
+
+	FString SupportedExtensions( FString::Printf( TEXT( "All Files (%s)|%s|%s" ), *AllExtensions, *AllExtensions, *FileTypes ) );
+
+	TArray<FString> OpenedFiles;
+	FString DefaultLocation( FEditorDirectories::Get().GetLastDirectory( ELastDirectory::GENERIC_IMPORT ) );
+	IDesktopPlatform* DesktopPlatform = FDesktopPlatformModule::Get();
+
+	bool bOpened = false;
+	if ( DesktopPlatform )
+	{
+		bOpened = DesktopPlatform->OpenFileDialog(
+			FSlateApplication::Get().FindBestParentWindowHandleForDialogs( nullptr ),
+			LOCTEXT( "FileDialogTitle", "Import Datasmith" ).ToString(),
+			DefaultLocation,
+			TEXT( "" ),
+			SupportedExtensions,
+			EFileDialogFlags::None,
+			OpenedFiles
+		);
+	}
+
+	if ( bOpened && OpenedFiles.Num() > 0 )
+	{
+		const FString& OpenedFile = OpenedFiles[0];
+		FEditorDirectories::Get().SetLastDirectory( ELastDirectory::GENERIC_IMPORT, OpenedFile );
+
+		return FPaths::ConvertRelativePathToFull( OpenedFile );
+	}
+
+	return FString();
+}
+
+FString FDatasmithFileProducerUtils::SelectDirectory()
+{
+	if( IDesktopPlatform* DesktopPlatform = FDesktopPlatformModule::Get() )
+	{
+		FString DestinationFolder;
+		const void* ParentWindowHandle = FSlateApplication::Get().FindBestParentWindowHandleForDialogs(nullptr);
+		const FString Title = LOCTEXT("DatasmithDirProducerFolderTitle", "Choose a folder").ToString();
+		const FString DefaultLocation( FEditorDirectories::Get().GetLastDirectory( ELastDirectory::GENERIC_IMPORT ) );
+
+		// #ueent_todo: What if the user selects a drive? Shouldn't we ask if this is really what he/she wants?
+		const bool bFolderSelected = DesktopPlatform->OpenDirectoryDialog(
+			ParentWindowHandle,
+			Title,
+			DefaultLocation,
+			DestinationFolder
+		);
+
+		if( bFolderSelected )
+		{
+			FEditorDirectories::Get().SetLastDirectory(ELastDirectory::GENERIC_EXPORT, DestinationFolder);
+
+			return FPaths::ConvertRelativePathToFull( DestinationFolder );
+		}
+	}
+
+	return FString();
+}
+
 #undef LOCTEXT_NAMESPACE
diff --git a/Engine/Plugins/Enterprise/DatasmithImporter/Source/DatasmithImporter/Private/DatasmithFileProducer.h b/Engine/Plugins/Enterprise/DatasmithImporter/Source/DatasmithImporter/Private/DatasmithFileProducer.h
--- a/Engine/Plugins/Enterprise/DatasmithImporter/Source/DatasmithImporter/Private/DatasmithFileProducer.h
+++ b/Engine/Plugins/Enterprise/DatasmithImporter/Source/DatasmithImporter/Private/DatasmithFileProducer.h
@@ -26,6 +26,10 @@ class DATASMITHIMPORTER_API UDatasmithFileProducer : public UDataprepContentProd
 	GENERATED_BODY()
 
 public:
+	// UObject interface
+	virtual void PostEditUndo() override;
+	virtual void PostInitProperties() override;
+	// End of UObject interface
 
 	/** Update producer with newly selected filename */
 	void SetFilename( const FString& InFilename );
@@ -70,6 +74,7 @@ private:
 	static FDatasmithImportBaseOptions DefaultImportOptions;
 
 	friend class SDatasmithFileProducerFileProperty;
+	friend class UDatasmithDirProducer;
 };
 
 // Customization of the details of the Datasmith Scene for the data prep editor.
@@ -94,6 +99,9 @@ public:
 
 	// Begin UObject interface
 	virtual void Serialize( FArchive& Ar ) override;
+	virtual void PostInitProperties() override;
+	// End of UObject interface
+
 	// End UObject interface
 
 	/** Update producer with newly selected folder name */
diff --git a/Engine/Plugins/Experimental/DataPrepEditor/Source/DataPrepCore/Private/DataPrepAssetProducers.cpp b/Engine/Plugins/Experimental/DataPrepEditor/Source/DataPrepCore/Private/DataPrepAssetProducers.cpp
--- a/Engine/Plugins/Experimental/DataPrepEditor/Source/DataPrepCore/Private/DataPrepAssetProducers.cpp
+++ b/Engine/Plugins/Experimental/DataPrepEditor/Source/DataPrepCore/Private/DataPrepAssetProducers.cpp
@@ -32,6 +32,13 @@ void UDataprepAssetProducers::Serialize( FArchive& Ar )
 	}
 }
 
+void UDataprepAssetProducers::PostEditUndo()
+{
+	Super::PostEditUndo();
+
+	OnChanged.Broadcast( FDataprepAssetChangeType::ProducerModified, INDEX_NONE );
+}
+
 TArray< TWeakObjectPtr< UObject > > UDataprepAssetProducers::Produce(const FDataprepProducerContext& InContext )
 {
 	if( AssetProducers.Num() == 0 )
@@ -85,6 +92,8 @@ UDataprepContentProducer* UDataprepAssetProducers::AddProducer(UClass* ProducerC
 {
 	if( ProducerClass && ProducerClass->IsChildOf( UDataprepContentProducer::StaticClass() ) )
 	{
+		Modify();
+
 		UDataprepContentProducer* Producer = NewObject< UDataprepContentProducer >( this, ProducerClass, NAME_None, RF_Transactional );
 		FAssetRegistryModule::AssetCreated( Producer );
 
@@ -92,8 +101,6 @@ UDataprepContentProducer* UDataprepAssetProducers::AddProducer(UClass* ProducerC
 
 		Producer->MarkPackageDirty();
 
-		Modify();
-
 		int32 ProducerNextIndex = AssetProducers.Num();
 		AssetProducers.Emplace( Producer, true );
 
@@ -152,56 +159,62 @@ bool UDataprepAssetProducers::RemoveProducer(int32 IndexToRemove)
 {
 	if( AssetProducers.IsValidIndex( IndexToRemove ) )
 	{
-		if(UDataprepContentProducer* Producer = AssetProducers[IndexToRemove].Producer)
-		{
-			Producer->GetOnChanged().RemoveAll( this );
-
-			DataprepCorePrivateUtils::DeleteRegisteredAsset( Producer );
-		}
-
-		Modify();
-
-		AssetProducers.RemoveAt( IndexToRemove );
+		// Update superseding status for producers depending on removed producer
+		bool bChangeAll = false;
 
 		// Array of producers superseded by removed producer
 		TArray<int32> ProducersToRevisit;
 		ProducersToRevisit.Reserve( AssetProducers.Num() );
-
-		if( AssetProducers.Num() == 1 )
-		{
-			AssetProducers[0].SupersededBy = INDEX_NONE;
-		}
-		else if(AssetProducers.Num() > 1)
 		{
-			// Update value stored in SupersededBy property where applicable
-			for( int32 Index = 0; Index < AssetProducers.Num(); ++Index )
+			Modify();
+
+			if(UDataprepContentProducer* Producer = AssetProducers[IndexToRemove].Producer)
 			{
-				FDataprepAssetProducer& AssetProducer = AssetProducers[Index];
+				Producer->GetOnChanged().RemoveAll( this );
 
-				if( AssetProducer.SupersededBy == IndexToRemove )
-				{
-					AssetProducer.SupersededBy = INDEX_NONE;
-					ProducersToRevisit.Add( Index );
-				}
-				else if( AssetProducer.SupersededBy > IndexToRemove )
+				Producer->PreEditChange(nullptr);
+				Producer->Modify();
+
+				DataprepCorePrivateUtils::DeleteRegisteredAsset( Producer );
+
+				Producer->PostEditChange();
+			}
+
+			AssetProducers.RemoveAt( IndexToRemove );
+
+			if( AssetProducers.Num() == 1 )
+			{
+				AssetProducers[0].SupersededBy = INDEX_NONE;
+			}
+			else if(AssetProducers.Num() > 1)
+			{
+				// Update value stored in SupersededBy property where applicable
+				for( int32 Index = 0; Index < AssetProducers.Num(); ++Index )
 				{
-					--AssetProducer.SupersededBy;
+					FDataprepAssetProducer& AssetProducer = AssetProducers[Index];
+
+					if( AssetProducer.SupersededBy == IndexToRemove )
+					{
+						AssetProducer.SupersededBy = INDEX_NONE;
+						ProducersToRevisit.Add( Index );
+					}
+					else if( AssetProducer.SupersededBy > IndexToRemove )
+					{
+						--AssetProducer.SupersededBy;
+					}
 				}
 			}
+
+			for( int32 ProducerIndex : ProducersToRevisit )
+			{
+				bool bLocalChangeAll = false;
+				ValidateProducerChanges( ProducerIndex, bLocalChangeAll );
+				bChangeAll |= bLocalChangeAll;
+			}
 		}
 
 		OnChanged.Broadcast( FDataprepAssetChangeType::ProducerRemoved, IndexToRemove );
 
-		// Update superseding status for producers depending on removed producer
-		bool bChangeAll = false;
-
-		for( int32 ProducerIndex : ProducersToRevisit )
-		{
-			bool bLocalChangeAll = false;
-			ValidateProducerChanges( ProducerIndex, bLocalChangeAll );
-			bChangeAll |= bLocalChangeAll;
-		}
-
 		// Notify observes on additional changes
 		if( bChangeAll )
 		{
diff --git a/Engine/Plugins/Experimental/DataPrepEditor/Source/DataPrepCore/Private/DataPrepContentConsumer.cpp b/Engine/Plugins/Experimental/DataPrepEditor/Source/DataPrepCore/Private/DataPrepContentConsumer.cpp
--- a/Engine/Plugins/Experimental/DataPrepEditor/Source/DataPrepCore/Private/DataPrepContentConsumer.cpp
+++ b/Engine/Plugins/Experimental/DataPrepEditor/Source/DataPrepCore/Private/DataPrepContentConsumer.cpp
@@ -23,6 +23,12 @@ UDataprepContentConsumer::UDataprepContentConsumer()
 	TargetContentFolder = FPaths::GetPath( GetOutermost()->GetPathName() );
 }
 
+void UDataprepContentConsumer::PostEditUndo()
+{
+	Super::PostEditUndo();
+	OnChanged.Broadcast();
+}
+
 bool UDataprepContentConsumer::Consume(const FDataprepConsumerContext& InContext)
 {
 	if(!InContext.WorldPtr.IsValid())
@@ -59,25 +65,26 @@ bool UDataprepContentConsumer::Consume(const FDataprepConsumerContext& InContext
 	return bSuccessful;
 }
 
-bool UDataprepContentConsumer::SetTargetContentFolder(const FString& InTargetContentFolder)
+bool UDataprepContentConsumer::SetTargetContentFolder(const FString& InTargetContentFolder, FText& OutReason)
 {
-	if( InTargetContentFolder.IsEmpty() )
+	bool bValidContentFolder = true;
+
+	if( !InTargetContentFolder.IsEmpty() )
 	{
-		TargetContentFolder = FPaths::GetPath( GetOutermost()->GetPathName() );
-		OnChanged.Broadcast();
-		return true;
+		// Pretend creating a dummy package to verify packages could be created under this content folder.
+		bValidContentFolder = FPackageName::IsValidLongPackageName( InTargetContentFolder / TEXT("DummyPackageName"), false, &OutReason );
 	}
 
-	// Pretend creating a dummy package to verify packages could be created under this content folder.
-	FString LongPackageName = InTargetContentFolder / TEXT("DummyPackageName");
-	if( FPackageName::IsValidLongPackageName( LongPackageName ) )
+	if(bValidContentFolder)
 	{
-		TargetContentFolder = InTargetContentFolder;
+		Modify();
+
+		TargetContentFolder = !InTargetContentFolder.IsEmpty() ? InTargetContentFolder : FPaths::GetPath( GetOutermost()->GetPathName() );
+
 		OnChanged.Broadcast();
-		return true;
 	}
 
-	return false;
+	return bValidContentFolder;
 }
 
 bool UDataprepContentConsumer::SetLevelName(const FString & InLevelName, FText& OutReason)
diff --git a/Engine/Plugins/Experimental/DataPrepEditor/Source/DataPrepCore/Private/Parameterization/DataprepParameterization.cpp b/Engine/Plugins/Experimental/DataPrepEditor/Source/DataPrepCore/Private/Parameterization/DataprepParameterization.cpp
--- a/Engine/Plugins/Experimental/DataPrepEditor/Source/DataPrepCore/Private/Parameterization/DataprepParameterization.cpp
+++ b/Engine/Plugins/Experimental/DataPrepEditor/Source/DataPrepCore/Private/Parameterization/DataprepParameterization.cpp
@@ -2,7 +2,7 @@
 
 #include "Parameterization/DataprepParameterization.h"
 
-#include "DataprepAsset.h"
+#include "DataPrepAsset.h"
 #include "DataprepParameterizationArchive.h"
 
 #include "CoreGlobals.h"
diff --git a/Engine/Plugins/Experimental/DataPrepEditor/Source/DataPrepCore/Public/DataPrepContentConsumer.h b/Engine/Plugins/Experimental/DataPrepEditor/Source/DataPrepCore/Public/DataPrepContentConsumer.h
--- a/Engine/Plugins/Experimental/DataPrepEditor/Source/DataPrepCore/Public/DataPrepContentConsumer.h
+++ b/Engine/Plugins/Experimental/DataPrepEditor/Source/DataPrepCore/Public/DataPrepContentConsumer.h
@@ -76,6 +76,10 @@ public:
 
 	UDataprepContentConsumer();
 
+	// UObject interface
+	virtual void PostEditUndo() override;
+	// End of UObject interface
+
 	/**
 	 * Successively calls Initialize, Run and Reset.
 	 * @param InContext : The world the consumer must process on. This world must be assumed to be transient.
@@ -111,7 +115,7 @@ public:
 	 * @return true if the assignment has been successful, false otherwise
 	 * @remark if InPackagePath is empty the package path of the consumer is used
 	 */
-	virtual bool SetTargetContentFolder(const FString& InTargetContentFolder );
+	virtual bool SetTargetContentFolder(const FString& InTargetContentFolder, FText& OutReason );
 
 	const FString& GetLevelName() { return LevelName; }
 
diff --git a/Engine/Plugins/Experimental/DataPrepEditor/Source/DataPrepCore/Public/DataprepAssetProducers.h b/Engine/Plugins/Experimental/DataPrepEditor/Source/DataPrepCore/Public/DataprepAssetProducers.h
--- a/Engine/Plugins/Experimental/DataPrepEditor/Source/DataPrepCore/Public/DataprepAssetProducers.h
+++ b/Engine/Plugins/Experimental/DataPrepEditor/Source/DataPrepCore/Public/DataprepAssetProducers.h
@@ -61,6 +61,7 @@ public:
 
 	// UObject interface
 	virtual void Serialize( FArchive& Ar ) override;
+	virtual void PostEditUndo() override;
 	// End of UObject interface
 
 	/**
diff --git a/Engine/Plugins/Experimental/DataPrepEditor/Source/DataPrepEditor/Private/Widgets/DataprepWidgets.cpp b/Engine/Plugins/Experimental/DataPrepEditor/Source/DataPrepEditor/Private/Widgets/DataprepWidgets.cpp
--- a/Engine/Plugins/Experimental/DataPrepEditor/Source/DataPrepEditor/Private/Widgets/DataprepWidgets.cpp
+++ b/Engine/Plugins/Experimental/DataPrepEditor/Source/DataPrepEditor/Private/Widgets/DataprepWidgets.cpp
@@ -5,6 +5,7 @@
 #include "DataPrepAsset.h"
 #include "DataPrepContentConsumer.h"
 #include "DataprepAssetView.h"
+#include "DataprepEditorLogCategory.h"
 #include "DataprepEditorUtils.h"
 #include "Parameterization/DataprepParameterizationUtils.h"
 
@@ -70,24 +71,56 @@ public:
 
 void SDataprepConsumerWidget::OnLevelNameChanged( const FText &NewLevelName, ETextCommit::Type CommitType )
 {
+	const FScopedTransaction Transaction( LOCTEXT("Consumer_SetLevelName", "Set Level Name") );
 	FText OutReason;
-	if( !DataprepConsumer->SetLevelName( NewLevelName.ToString(), OutReason ) )
+	if( !DataprepConsumerPtr->SetLevelName( NewLevelName.ToString(), OutReason ) )
 	{
 		// #ueent_todo: Warn user name is wrong
-		LevelTextBox->SetText( FText::FromString( DataprepConsumer->GetLevelName() ) );
+		LevelTextBox->SetText( FText::FromString( DataprepConsumerPtr->GetLevelName() ) );
+	}
+}
+
+void SDataprepConsumerWidget::OnTextCommitted( const FText& NewText, ETextCommit::Type CommitType)
+{
+	if( UDataprepContentConsumer* DataprepConsumer = DataprepConsumerPtr.Get() )
+	{
+		FString NewContentFolder( NewText.ToString() );
+
+		// Replace /Content/ with /Game/ since /Content is only used for display 
+		if( NewContentFolder.StartsWith( TEXT("/Content") ) )
+		{
+			NewContentFolder = NewContentFolder.Replace( TEXT( "/Content" ), TEXT( "/Game" ) );
+		}
+		
+		// Remove ending '/' if applicable
+		if( !NewContentFolder.IsEmpty() && NewContentFolder[ NewContentFolder.Len()-1 ] == TEXT('/') )
+		{
+			NewContentFolder[ NewContentFolder.Len()-1 ] = 0;
+			NewContentFolder = NewContentFolder.LeftChop(1);
+		}
+
+		const FScopedTransaction Transaction( LOCTEXT("Consumer_SetTargetContentFolder", "Set Target Content Folder") );
+
+		FText ErrorReason;
+		if( !DataprepConsumer->SetTargetContentFolder( NewContentFolder, ErrorReason ) )
+		{
+			UE_LOG( LogDataprepEditor, Error, TEXT("%s"), *ErrorReason.ToString() );
+			UpdateContentFolderText();
+		}
 	}
 }
 
 void SDataprepConsumerWidget::OnBrowseContentFolder()
 {
-	FString Path = DataprepConsumer->GetTargetContentFolder();
+	FString Path = DataprepConsumerPtr->GetTargetContentFolder();
 	if( Path.IsEmpty() )
 	{
-		Path = FPaths::GetPath( DataprepConsumer->GetOutermost()->GetPathName() );
+		Path = FPaths::GetPath( DataprepConsumerPtr->GetOutermost()->GetPathName() );
 	}
 	Path += TEXT("/"); // Trailing '/' is needed to set the default path
 
 	//Ask the user for the root path where they want any content to be placed
+	if( UDataprepContentConsumer* DataprepConsumer = DataprepConsumerPtr.Get() )
 	{
 		TSharedRef<SDlgPickPath> PickContentPathDlg =
 			SNew(SDlgPickPath)
@@ -96,8 +129,17 @@ void SDataprepConsumerWidget::OnBrowseContentFolder()
 
 		if ( PickContentPathDlg->ShowModal() == EAppReturnType::Ok )
 		{
-			DataprepConsumer->SetTargetContentFolder( PickContentPathDlg->GetPath().ToString() );
-			UpdateContentFolderText();
+			const FScopedTransaction Transaction( LOCTEXT("Consumer_SetTargetContentFolder", "Set Target Content Folder") );
+
+			FText ErrorReason;
+			if( DataprepConsumer->SetTargetContentFolder( PickContentPathDlg->GetPath().ToString(), ErrorReason ) )
+			{
+				UpdateContentFolderText();
+			}
+			else
+			{
+				UE_LOG( LogDataprepEditor, Error, TEXT("%s"), *ErrorReason.ToString() );
+			}
 		}
 	}
 }
@@ -109,16 +151,23 @@ void SDataprepConsumerWidget::SetDataprepConsumer(UDataprepContentConsumer* InDa
 		return;
 	}
 
-	DataprepConsumer = InDataprepConsumer;
+	if(DataprepConsumerPtr != InDataprepConsumer)
+	{
+		if(DataprepConsumerPtr.IsValid())
+		{
+			DataprepConsumerPtr->GetOnChanged().RemoveAll( this );
+		}
+
+		DataprepConsumerPtr = InDataprepConsumer;
+
+		InDataprepConsumer->GetOnChanged().AddRaw( this, &SDataprepConsumerWidget::OnConsumerChanged );
 
-	UpdateContentFolderText();
-	LevelTextBox->SetText( FText::FromString( DataprepConsumer->GetLevelName() ) );
+		OnConsumerChanged();
+	}
 }
 
 void SDataprepConsumerWidget::Construct(const FArguments& InArgs )
 {
-	DataprepConsumer = InArgs._DataprepConsumer;
-
 	if (InArgs._ColumnSizeData.IsValid())
 	{
 		ColumnSizeData = InArgs._ColumnSizeData;
@@ -132,19 +181,21 @@ void SDataprepConsumerWidget::Construct(const FArguments& InArgs )
 		ColumnSizeData->OnWidthChanged = SSplitter::FOnSlotResized::CreateSP(this, &SDataprepConsumerWidget::OnSetColumnWidth);
 	}
 
+	TSharedPtr<SWidget> InnerWidget = InArgs._DataprepConsumer ? BuildWidget() : BuildNullWidget();
+
 	ChildSlot
 	[
-		BuildWidget()
+		InnerWidget.ToSharedRef()
 	];
-}
 
-TSharedRef<SWidget> SDataprepConsumerWidget::BuildWidget()
-{
-	if(DataprepConsumer == nullptr)
+	if(InArgs._DataprepConsumer)
 	{
-		return BuildNullWidget();
+		SetDataprepConsumer( InArgs._DataprepConsumer );
 	}
+}
 
+TSharedRef<SWidget> SDataprepConsumerWidget::BuildWidget()
+{
 	TSharedRef<SWidget> BrowseButton = PropertyCustomizationHelpers::MakeBrowseButton( FSimpleDelegate::CreateSP( this, &SDataprepConsumerWidget::OnBrowseContentFolder ) );
 
 	TSharedRef<SWidget> Widget = SNew(SBorder)
@@ -199,6 +250,7 @@ TSharedRef<SWidget> SDataprepConsumerWidget::BuildWidget()
 							.Font(IDetailLayoutBuilder::GetDetailFont())
 							.HintText(LOCTEXT("DataprepSlateHelper_ContentFolderHintText", "Set the content folder to save in"))
 							.IsReadOnly(false)
+							.OnTextCommitted(FOnTextCommitted::CreateSP(this, &SDataprepConsumerWidget::OnTextCommitted))
 						]
 					]
 				]
@@ -284,7 +336,6 @@ TSharedRef<SWidget> SDataprepConsumerWidget::BuildWidget()
 						SNew(SConstrainedBox)
 						[
 							SAssignNew(LevelTextBox, SEditableTextBox)
-							.Text(FText::FromString(DataprepConsumer->GetLevelName()))
 							.Font(IDetailLayoutBuilder::GetDetailFont())
 							.HintText(LOCTEXT("DataprepSlateHelper_LevelNameHintText", "Current will be used"))
 							.OnTextCommitted(this, &SDataprepConsumerWidget::OnLevelNameChanged)
@@ -295,8 +346,6 @@ TSharedRef<SWidget> SDataprepConsumerWidget::BuildWidget()
 		]
 	];
 
-	UpdateContentFolderText();
-
 	return Widget;
 }
 
@@ -327,7 +376,7 @@ TSharedRef<SWidget> SDataprepConsumerWidget::BuildNullWidget()
 
 void SDataprepConsumerWidget::UpdateContentFolderText()
 {
-	if(UDataprepContentConsumer* Consumer = DataprepConsumer.Get())
+	if(UDataprepContentConsumer* Consumer = DataprepConsumerPtr.Get())
 	{
 		FString TargetContentFolder( Consumer->GetTargetContentFolder() );
 
@@ -342,6 +391,24 @@ void SDataprepConsumerWidget::UpdateContentFolderText()
 
 		ContentFolderTextBox->SetText( FText::FromString( TargetContentFolder + TEXT( "/" ) ) );
 	}
+	else
+	{
+		ContentFolderTextBox->SetText( TAttribute<FText>() );
+	}
+}
+
+void SDataprepConsumerWidget::OnConsumerChanged()
+{
+	if(DataprepConsumerPtr.IsValid())
+	{
+		UpdateContentFolderText();
+		LevelTextBox->SetText( FText::FromString( DataprepConsumerPtr->GetLevelName() ) );
+	}
+	else
+	{
+		ContentFolderTextBox->SetText( TAttribute<FText>() );
+		LevelTextBox->SetText( TAttribute<FText>() );
+	}
 }
 
 
diff --git a/Engine/Plugins/Experimental/DataPrepEditor/Source/DataPrepEditor/Private/Widgets/DataprepWidgets.h b/Engine/Plugins/Experimental/DataPrepEditor/Source/DataPrepEditor/Private/Widgets/DataprepWidgets.h
--- a/Engine/Plugins/Experimental/DataPrepEditor/Source/DataPrepEditor/Private/Widgets/DataprepWidgets.h
+++ b/Engine/Plugins/Experimental/DataPrepEditor/Source/DataPrepEditor/Private/Widgets/DataprepWidgets.h
@@ -28,6 +28,7 @@ class SDataprepConsumerWidget : public SCompoundWidget
 {
 public:
 	SLATE_BEGIN_ARGS(SDataprepConsumerWidget)
+		: _DataprepConsumer(nullptr)
 	{}
 
 	SLATE_ARGUMENT(TSharedPtr< FDataprepDetailsViewColumnSizeData >, ColumnSizeData)
@@ -37,15 +38,27 @@ public:
 public:
 	void Construct(const FArguments& InArgs );
 
+	/** Update DataprepConsumerPtr and content folder's and level name's text boxes */
 	void SetDataprepConsumer( UDataprepContentConsumer* DataprepConsumer );
 
 private:
+	/** Create complete widget from valid UDataprepContentConsumer object */
 	TSharedRef<SWidget> BuildWidget();
+	/** Create empty widget from invalid UDataprepContentConsumer object */
 	TSharedRef<SWidget> BuildNullWidget();
+	/** Callback when content of level name's text box has changed */
 	void OnLevelNameChanged( const FText& NewLevelName, ETextCommit::Type CommitType );
+	/** Callback when browser's button is clicked */
 	void OnBrowseContentFolder();
+	/** Callback when content of content folder's text box has changed */
+	void OnTextCommitted( const FText&, ETextCommit::Type );
+
+	/** Update content folder's and level name's text boxes */
+	void OnConsumerChanged();
+	/** Update content folder's text boxes */
 	void UpdateContentFolderText();
 
+	/** Callbacks to update splitter */
 	void OnLeftColumnResized(float InNewWidth)
 	{
 		// This has to be bound or the splitter will take it upon itself to determine the size
@@ -57,10 +70,13 @@ private:
 	void OnSetColumnWidth(float InWidth) { ColumnWidth = InWidth; }
 
 private:
-	TWeakObjectPtr<UDataprepContentConsumer> DataprepConsumer;
+	/** Weak pointer on edited consumer */
+	TWeakObjectPtr<UDataprepContentConsumer> DataprepConsumerPtr;
+	/** Content folder's text box */
 	TSharedPtr< SEditableTextBox > ContentFolderTextBox;
+	/** Level name's text box */
 	TSharedPtr< SEditableTextBox > LevelTextBox;
-	// Helps sync column resizing with another part of the UI (producers widget).
+	/** Helps sync column resizing with another part of the UI (producers widget) */
 	TSharedPtr< FDataprepDetailsViewColumnSizeData > ColumnSizeData;
 	/** Relative width to control splitters. */
 	float ColumnWidth;
diff --git a/Engine/Plugins/Experimental/DataPrepEditor/Source/DataPrepEditor/Private/Widgets/SDataprepProducersWidget.cpp b/Engine/Plugins/Experimental/DataPrepEditor/Source/DataPrepEditor/Private/Widgets/SDataprepProducersWidget.cpp
--- a/Engine/Plugins/Experimental/DataPrepEditor/Source/DataPrepEditor/Private/Widgets/SDataprepProducersWidget.cpp
+++ b/Engine/Plugins/Experimental/DataPrepEditor/Source/DataPrepEditor/Private/Widgets/SDataprepProducersWidget.cpp
@@ -12,6 +12,7 @@
 #include "IDetailTreeNode.h"
 #include "K2Node_AddComponent.h"
 #include "PropertyEditorModule.h"
+#include "ScopedTransaction.h"
 
 #include "DetailCategoryBuilder.h"
 #include "DetailLayoutBuilder.h"
@@ -51,6 +52,29 @@ FContentProducerEntry::FContentProducerEntry(int32 InProducerIndex, UDataprepAss
 	}
 }
 
+
+void FContentProducerEntry::ToggleProducer()
+{
+	if( UDataprepAssetProducers* AssetProducers = AssetProducersPtr.Get() )
+	{
+		const FScopedTransaction Transaction( LOCTEXT("Producers_ToggleProducer", "Toggle Producer") );
+		AssetProducers->EnableProducer(ProducerIndex, !bIsEnabled);
+
+		// #ueent_todo: Cache previous value to report failed enabling/disabling
+		bIsEnabled = AssetProducers->IsProducerEnabled(ProducerIndex);
+	}
+}
+
+void FContentProducerEntry::RemoveProducer()
+{
+	if( UDataprepAssetProducers* AssetProducers = AssetProducersPtr.Get() )
+	{
+		const FScopedTransaction Transaction( LOCTEXT("Producers_RemoveProducer", "Remove Producer") );
+		AssetProducers->RemoveProducer( ProducerIndex );
+	}
+}
+
+
 /** Construct function for this widget */
 void SDataprepProducersTableRow::Construct(const FArguments& InArgs, const TSharedRef<STableViewBase>& OwnerTableView, const TSharedRef<FContentProducerEntry>& InNode, TSharedRef< FDataprepDetailsViewColumnSizeData > InColumnSizeData)
 {
@@ -314,10 +338,7 @@ void SDataprepProducersWidget::Construct(const FArguments & InArgs, UDataprepAss
 
 void SDataprepProducersWidget::OnDataprepProducersChanged(FDataprepAssetChangeType ChangeType, int32 Index)
 {
-	if(ChangeType == FDataprepAssetChangeType::ProducerRemoved)
-	{
-		TreeView->Refresh();
-	}
+	TreeView->Refresh();
 }
 
 TSharedRef<SWidget> SDataprepProducersWidget::CreateAddProducerMenuWidget()
@@ -367,10 +388,11 @@ void SDataprepProducersWidget::OnAddProducer( UClass* ProducerClass )
 {
 	if( UDataprepAssetProducers* AssetProducers = AssetProducersPtr.Get() )
 	{
+		const FScopedTransaction Transaction( LOCTEXT("Producers_AddProducer", "Add Producer") );
 		AssetProducers->AddProducer(ProducerClass);
 	}
-	TreeView->
-	Refresh();
+
+	TreeView->Refresh();
 }
 
 TSharedRef<SWidget> FDataprepAssetProducersDetails::CreateWidget(UDataprepAssetProducers* Producers, TSharedPtr<FUICommandList>& CommandList)
diff --git a/Engine/Plugins/Experimental/DataPrepEditor/Source/DataPrepEditor/Private/Widgets/SDataprepProducersWidget.h b/Engine/Plugins/Experimental/DataPrepEditor/Source/DataPrepEditor/Private/Widgets/SDataprepProducersWidget.h
--- a/Engine/Plugins/Experimental/DataPrepEditor/Source/DataPrepEditor/Private/Widgets/SDataprepProducersWidget.h
+++ b/Engine/Plugins/Experimental/DataPrepEditor/Source/DataPrepEditor/Private/Widgets/SDataprepProducersWidget.h
@@ -35,24 +35,9 @@ public:
 
 	bool WillBeRun() { return bIsEnabled && !bIsSuperseded; }
 
-	void ToggleProducer()
-	{
-		if( UDataprepAssetProducers* AssetProducers = AssetProducersPtr.Get() )
-		{
-			AssetProducers->EnableProducer(ProducerIndex, !bIsEnabled);
-
-			// #ueent_todo: Cache previous value to report failed enabling/disabling
-			bIsEnabled = AssetProducers->IsProducerEnabled(ProducerIndex);
-		}
-	}
+	void ToggleProducer();
 
-	void RemoveProducer()
-	{
-		if( UDataprepAssetProducers* AssetProducers = AssetProducersPtr.Get() )
-		{
-			AssetProducers->RemoveProducer( ProducerIndex );
-		}
-	}
+	void RemoveProducer();
 
 	FString Label;
 	int32 ProducerIndex;
