Integrating CL# 9465487 from Dev-Enterprise
Adding option to generate lightmaps during datasmith import

diff --git a/Engine/Plugins/Enterprise/DatasmithContent/Source/DatasmithContent/Public/DatasmithImportOptions.h b/Engine/Plugins/Enterprise/DatasmithContent/Source/DatasmithContent/Public/DatasmithImportOptions.h
--- a/Engine/Plugins/Enterprise/DatasmithContent/Source/DatasmithContent/Public/DatasmithImportOptions.h
+++ b/Engine/Plugins/Enterprise/DatasmithContent/Source/DatasmithContent/Public/DatasmithImportOptions.h
@@ -134,7 +134,7 @@ struct DATASMITHCONTENT_API FDatasmithStaticMeshImportOptions
 	UPROPERTY(EditAnywhere, BlueprintReadWrite, Category = Lightmap)
 	EDatasmithImportLightmapMax MaxLightmapResolution;
 
-	UPROPERTY(BlueprintReadWrite, Category = Lightmap, Transient)
+	UPROPERTY(EditAnywhere, BlueprintReadWrite, Category = Lightmap, meta = (DisplayName = "Generate Lightmap UVs"))
 	bool bGenerateLightmapUVs;
 
 	UPROPERTY(BlueprintReadWrite, Category = Mesh, Transient)
diff --git a/Engine/Plugins/Enterprise/DatasmithImporter/DatasmithImporter.uplugin b/Engine/Plugins/Enterprise/DatasmithImporter/DatasmithImporter.uplugin
--- a/Engine/Plugins/Enterprise/DatasmithImporter/DatasmithImporter.uplugin
+++ b/Engine/Plugins/Enterprise/DatasmithImporter/DatasmithImporter.uplugin
@@ -46,6 +46,10 @@
 		{
 			"Name": "VariantManager",
 			"Enabled": true
+		},
+		{
+			"Name": "StaticMeshEditorExtension",
+			"Enabled": true
 		}
 	]
 }
\ No newline at end of file
diff --git a/Engine/Plugins/Enterprise/DatasmithImporter/Source/DatasmithImporter/DatasmithImporter.Build.cs b/Engine/Plugins/Enterprise/DatasmithImporter/Source/DatasmithImporter/DatasmithImporter.Build.cs
--- a/Engine/Plugins/Enterprise/DatasmithImporter/Source/DatasmithImporter/DatasmithImporter.Build.cs
+++ b/Engine/Plugins/Enterprise/DatasmithImporter/Source/DatasmithImporter/DatasmithImporter.Build.cs
@@ -44,7 +44,8 @@ namespace UnrealBuildTool.Rules
 					"UnrealEd",
 					"VariantManager",
 					"VariantManagerContent",
-				}
+                    "StaticMeshEditorExtension",
+                }
 			);
 
 			PublicDependencyModuleNames.AddRange(
diff --git a/Engine/Plugins/Enterprise/DatasmithImporter/Source/DatasmithImporter/Private/DatasmithStaticMeshImporter.cpp b/Engine/Plugins/Enterprise/DatasmithImporter/Source/DatasmithImporter/Private/DatasmithStaticMeshImporter.cpp
--- a/Engine/Plugins/Enterprise/DatasmithImporter/Source/DatasmithImporter/Private/DatasmithStaticMeshImporter.cpp
+++ b/Engine/Plugins/Enterprise/DatasmithImporter/Source/DatasmithImporter/Private/DatasmithStaticMeshImporter.cpp
@@ -28,6 +28,7 @@
 #include "StaticMeshAttributes.h"
 #include "Templates/UniquePtr.h"
 #include "UObject/Package.h"
+#include "UVTools/UVGenerationFlattenMapping.h"
 
 #define LOCTEXT_NAMESPACE "DatasmithStaticMeshImporter"
 
@@ -202,20 +203,19 @@ bool FDatasmithStaticMeshImporter::PreBuildStaticMesh( UStaticMesh* StaticMesh )
 			DatasmithMeshHelper::CreateDefaultUVs(MeshDescription);
 		}
 
-		// If lightmap UVs are not to be generated, set destination lightmap to 0.
-		if (!BuildSettings.bGenerateLightmapUVs)
+		if (BuildSettings.bGenerateLightmapUVs && !DatasmithMeshHelper::HasUVData(MeshDescription, BuildSettings.SrcLightmapIndex))
 		{
-			BuildSettings.DstLightmapIndex = 0;
-		}
-		else if (!DatasmithMeshHelper::HasUVData(MeshDescription, BuildSettings.SrcLightmapIndex))
-		{
-			//We know that a UV at channel 0 exists so we fallback on it.
-			BuildSettings.SrcLightmapIndex = 0;
+			//If no UV data exist at the source index we generate unwrapped UVs.
+			UUVGenerationFlattenMapping::GenerateUVs(MeshDescription, BuildSettings.SrcLightmapIndex, BuildSettings.bRemoveDegenerates);
 		}
 
 		if (DatasmithMeshHelper::IsMeshValid(MeshDescription, BuildSettings.BuildScale3D))
 		{
-			DatasmithMeshHelper::RequireUVChannel(MeshDescription, BuildSettings.DstLightmapIndex);
+			if (BuildSettings.bGenerateLightmapUVs)
+			{
+				DatasmithMeshHelper::RequireUVChannel(MeshDescription, BuildSettings.DstLightmapIndex);
+			}
+
 			UStaticMesh::FCommitMeshDescriptionParams Params;
 
 			// We will call MarkPackageDirty() ourself from the main thread
@@ -480,8 +480,6 @@ void FDatasmithStaticMeshImporter::SetupStaticMesh( FDatasmithAssetsImportContex
 
 	TSet< int32 > MaterialIndexToImportIndex;
 
-	bool bLightmapCoordinateIndexInvalid = false;
-
 	// Only setup LODs declared by the mesh element
 	for ( int32 LodIndex = 0; LodIndex < MeshElement->GetLODCount(); ++LodIndex)
 	{
@@ -505,35 +503,47 @@ void FDatasmithStaticMeshImporter::SetupStaticMesh( FDatasmithAssetsImportContex
 		};
 
 		// UV Channels
-		int32 SourceIndex = -1;
+		int32 SourceIndex = 0;
+		int32 DestinationIndex = 1;
+		bool bUseImportedLightmap = false;
 		bool bGenerateLightmapUVs = StaticMeshImportOptions.bGenerateLightmapUVs;
-		// Set DestinationIndex to 0 if generation of lightmap UVs is not required.
-		int32 DestinationIndex = bGenerateLightmapUVs ? -1 : 0;
+		const int32 FirstOpenUVChannel = GetNextOpenUVChannel(MeshDescription, 0);
 
-		// if a lightmap coordinate index was exported, disable lightmap generation
+		// if a custom lightmap coordinate index was imported, disable lightmap generation
 		if (DatasmithMeshHelper::HasUVData(MeshDescription, MeshElement->GetLightmapCoordinateIndex()))
 		{
+			bUseImportedLightmap = true;
 			bGenerateLightmapUVs = false;
+			DestinationIndex = MeshElement->GetLightmapCoordinateIndex();
 		}
-
-		if (bGenerateLightmapUVs)
+		else
 		{
-			// Set the source lightmap index to the imported mesh data lightmap source if any, otherwise use the first open channel.
-			const bool bOriginalSourceIsValid = DatasmithMeshHelper::HasUVData(MeshDescription, MeshElement->GetLightmapSourceUV());
-			const int32 FirstOpenUVChannel = GetNextOpenUVChannel(MeshDescription, 0);
-
-			if (bOriginalSourceIsValid)
-			{
-				SourceIndex = MeshElement->GetLightmapSourceUV();
-				DestinationIndex = FirstOpenUVChannel;
-			}
-			else
+			if (MeshElement->GetLightmapCoordinateIndex() >= 0)
 			{
-				//If the lightmap source index was not set, set it to 0 (as we will generate a UV latter if it's missing).
-				SourceIndex = 0;
-				DestinationIndex = FMath::Max(FirstOpenUVChannel, SourceIndex + 1);
+				//A custom lightmap index value was set but the data was invalid.
+				FFormatNamedArguments FormatArgs;
+				FormatArgs.Add(TEXT("LightmapCoordinateIndex"), FText::FromString(FString::FromInt(MeshElement->GetLightmapCoordinateIndex())));
+				FormatArgs.Add(TEXT("MeshName"), FText::FromName(StaticMesh->GetFName()));
+				AssetsContext.ParentContext.LogError(FText::Format(LOCTEXT("InvalidLightmapSourceUVError", "The lightmap coordinate index '{LightmapCoordinateIndex}' used for the mesh '{MeshName}' is invalid. A valid lightmap coordinate index was set instead."), FormatArgs));
 			}
-			
+
+			DestinationIndex = FirstOpenUVChannel;
+		}
+
+		// Set the source lightmap index to the imported mesh data lightmap source if any, otherwise use the first open channel.
+		if (DatasmithMeshHelper::HasUVData(MeshDescription, MeshElement->GetLightmapSourceUV()))
+		{
+			SourceIndex = MeshElement->GetLightmapSourceUV();
+		}
+		else
+		{
+			//If the lightmap source index was not set, we set it to the first open UV channel as it will be generated.
+			//Also, it's okay to set both the source and the destination to be the same index as they are for different containers.
+			SourceIndex = FirstOpenUVChannel;
+		}
+		
+		if (bGenerateLightmapUVs)
+		{
 			if (!FMath::IsWithin<int32>(SourceIndex, 0, MAX_MESH_TEXTURE_COORDS_MD))
 			{
 				AssetsContext.ParentContext.LogError(FText::Format(LOCTEXT("InvalidLightmapSourceIndexError", "Lightmap generation error for mesh {0}: Specified source is invalid {1}. Cannot find an available fallback source channel."), FText::FromName(StaticMesh->GetFName()), MeshElement->GetLightmapSourceUV()));
@@ -547,12 +557,10 @@ void FDatasmithStaticMeshImporter::SetupStaticMesh( FDatasmithAssetsImportContex
 
 			if (!bGenerateLightmapUVs)
 			{
-				AssetsContext.ParentContext.LogWarning(FText::Format(LOCTEXT("LightmapUVsWontBeGenerated", "Lightmap UVs for mesh {0} won't be generated. The lightmap coordinate index will be set to 0."), FText::FromName(StaticMesh->GetFName())));
-				SourceIndex = 0;
-				DestinationIndex = 1;
+				AssetsContext.ParentContext.LogWarning(FText::Format(LOCTEXT("LightmapUVsWontBeGenerated", "Lightmap UVs for mesh {0} won't be generated."), FText::FromName(StaticMesh->GetFName())));
 			}
 		}
-
+		
 		FDatasmithMeshBuildSettingsTemplate BuildSettingsTemplate;
 		BuildSettingsTemplate.Load(StaticMesh->GetSourceModel(LodIndex).BuildSettings);
 
@@ -572,31 +580,15 @@ void FDatasmithStaticMeshImporter::SetupStaticMesh( FDatasmithAssetsImportContex
 
 		StaticMeshTemplate->BuildSettings.Add( MoveTemp( BuildSettingsTemplate ) );
 
-		if (!StaticMeshImportOptions.bGenerateLightmapUVs)
+		if ( bUseImportedLightmap || StaticMeshImportOptions.bGenerateLightmapUVs )
 		{
-			StaticMeshTemplate->LightMapCoordinateIndex = MeshElement->GetLightmapCoordinateIndex() < 0 ? BuildSettingsTemplate.DstLightmapIndex : MeshElement->GetLightmapCoordinateIndex();
-		}
-		else if (MeshElement->GetLightmapCoordinateIndex() < 0)
-		{
-			StaticMeshTemplate->LightMapCoordinateIndex = bGenerateLightmapUVs ? BuildSettingsTemplate.DstLightmapIndex : BuildSettingsTemplate.SrcLightmapIndex;
+			//If we are generating the lightmap or are using an imported lightmap the DstLightmapIndex will already be set to the proper index.
+			StaticMeshTemplate->LightMapCoordinateIndex = BuildSettingsTemplate.DstLightmapIndex;
 		}
 		else
 		{
-			int LastValidUVChannel = MAX_MESH_TEXTURE_COORDS_MD - 1;
-			while (LastValidUVChannel > 0 && !DatasmithMeshHelper::HasUVData(MeshDescription, LastValidUVChannel))
-			{
-				--LastValidUVChannel;
-			}
-
-			if (MeshElement->GetLightmapCoordinateIndex() <= LastValidUVChannel)
-			{
-				StaticMeshTemplate->LightMapCoordinateIndex = MeshElement->GetLightmapCoordinateIndex();
-			}
-			else
-			{
-				bLightmapCoordinateIndexInvalid = true;
-				StaticMeshTemplate->LightMapCoordinateIndex = LastValidUVChannel;
-			}
+			//No lightmap UV provided or generated, set the LightmapCoordinateIndex to default 0 value which will always contain contain a basic UV.
+			StaticMeshTemplate->LightMapCoordinateIndex = 0;
 		}
 	}
 
@@ -615,14 +607,6 @@ void FDatasmithStaticMeshImporter::SetupStaticMesh( FDatasmithAssetsImportContex
 
 	ApplyMaterialsToStaticMesh(AssetsContext, MeshElement, StaticMeshTemplate);
 	StaticMeshTemplate->Apply(StaticMesh);
-
-	if (bLightmapCoordinateIndexInvalid)
-	{
-		FFormatNamedArguments FormatArgs;
-		FormatArgs.Add(TEXT("LightmapCoordinateIndex"), FText::FromString(FString::FromInt(MeshElement->GetLightmapCoordinateIndex())));
-		FormatArgs.Add(TEXT("MeshName"), FText::FromName(StaticMesh->GetFName()));
-		AssetsContext.ParentContext.LogError(FText::Format(LOCTEXT("InvalidLightmapSourceUVError", "The lightmap coordinate index '{LightmapCoordinateIndex}' used for the mesh '{MeshName}' is invalid. A valid lightmap coordinate index was set instead."), FormatArgs));
-	}
 }
 
 TMap< TSharedRef< IDatasmithMeshElement >, float > FDatasmithStaticMeshImporter::CalculateMeshesLightmapWeights( const TSharedRef< IDatasmithScene >& SceneElement )
diff --git a/Engine/Plugins/Enterprise/StaticMeshEditorExtension/Source/StaticMeshEditorExtension/Private/UVTools/UVGenerationDataprepOperation.cpp b/Engine/Plugins/Enterprise/StaticMeshEditorExtension/Source/StaticMeshEditorExtension/Private/UVTools/UVGenerationDataprepOperation.cpp
--- a/Engine/Plugins/Enterprise/StaticMeshEditorExtension/Source/StaticMeshEditorExtension/Private/UVTools/UVGenerationDataprepOperation.cpp
+++ b/Engine/Plugins/Enterprise/StaticMeshEditorExtension/Source/StaticMeshEditorExtension/Private/UVTools/UVGenerationDataprepOperation.cpp
@@ -2,7 +2,7 @@
 
 #include "UVGenerationDataprepOperation.h"
 
-#include "UVGenerationFlattenMapping.h"
+#include "UVTools/UVGenerationFlattenMapping.h"
 
 #include "Components/StaticMeshComponent.h"
 #include "Engine/StaticMesh.h"
diff --git a/Engine/Plugins/Enterprise/StaticMeshEditorExtension/Source/StaticMeshEditorExtension/Private/UVTools/UVGenerationFlattenMapping.cpp b/Engine/Plugins/Enterprise/StaticMeshEditorExtension/Source/StaticMeshEditorExtension/Private/UVTools/UVGenerationFlattenMapping.cpp
--- a/Engine/Plugins/Enterprise/StaticMeshEditorExtension/Source/StaticMeshEditorExtension/Private/UVTools/UVGenerationFlattenMapping.cpp
+++ b/Engine/Plugins/Enterprise/StaticMeshEditorExtension/Source/StaticMeshEditorExtension/Private/UVTools/UVGenerationFlattenMapping.cpp
@@ -1,6 +1,6 @@
 // Copyright 1998-2019 Epic Games, Inc. All Rights Reserved.
 
-#include "UVGenerationFlattenMapping.h"
+#include "UVTools/UVGenerationFlattenMapping.h"
 
 #include "Algo/Sort.h"
 #include "Async/ParallelFor.h"
diff --git a/Engine/Plugins/Enterprise/StaticMeshEditorExtension/Source/StaticMeshEditorExtension/Private/UVTools/UVGenerationFlattenMapping.h b/Engine/Plugins/Enterprise/StaticMeshEditorExtension/Source/StaticMeshEditorExtension/Private/UVTools/UVGenerationFlattenMapping.h
deleted file mode 100644
diff --git a/Engine/Plugins/Enterprise/StaticMeshEditorExtension/Source/StaticMeshEditorExtension/Public/UVTools/UVGenerationFlattenMapping.h b/Engine/Plugins/Enterprise/StaticMeshEditorExtension/Source/StaticMeshEditorExtension/Public/UVTools/UVGenerationFlattenMapping.h
new file mode 100644
--- /dev/null
+++ b/Engine/Plugins/Enterprise/StaticMeshEditorExtension/Source/StaticMeshEditorExtension/Public/UVTools/UVGenerationFlattenMapping.h
@@ -0,0 +1,27 @@
+// Copyright 1998-2019 Epic Games, Inc. All Rights Reserved.
+
+#pragma once
+
+#include "CoreTypes.h"
+#include "Kismet/BlueprintFunctionLibrary.h"
+
+#include "UVGenerationFlattenMapping.generated.h"
+
+class UStaticMesh;
+struct FMeshDescription;
+
+UCLASS()
+class STATICMESHEDITOREXTENSION_API UUVGenerationFlattenMapping : public UBlueprintFunctionLibrary
+{
+	GENERATED_BODY()
+
+public:
+	UFUNCTION(BlueprintCallable, Category = "Editor Scripting | StaticMesh", meta = (DisplayName = "Generate Flatten Mapping UVs"))
+	static void GenerateFlattenMappingUVs(UStaticMesh* InStaticMesh, int32 UVChannel, float AngleThreshold = 66.f, float AreaWeight = 0.7f);
+
+	static void GenerateUVs(FMeshDescription& InMesh, int32 UVChannel, bool bRemoveDegenerate, float AngleThreshold = 66.f, float AreaWeight = 0.7f);
+
+private:
+
+	static TArray<int32> GetOverlappingCornersRemapping(const FMeshDescription& InMeshDescription, bool bRemoveDegenerates);
+};
