Copy of 9481076 from //UE4/Release-4.24:

Datasmith: Improve dependency on FreeImage

diff --git a/Engine/Plugins/Enterprise/DatasmithImporter/Source/DatasmithImporter/Private/Utility/DatasmithTextureResize.cpp b/Engine/Plugins/Enterprise/DatasmithImporter/Source/DatasmithImporter/Private/Utility/DatasmithTextureResize.cpp
--- a/Engine/Plugins/Enterprise/DatasmithImporter/Source/DatasmithImporter/Private/Utility/DatasmithTextureResize.cpp
+++ b/Engine/Plugins/Enterprise/DatasmithImporter/Source/DatasmithImporter/Private/Utility/DatasmithTextureResize.cpp
@@ -4,7 +4,7 @@
 
 #include "DatasmithUtils.h"
 
-#if PLATFORM_WINDOWS // #ueent_todo Datasmith texture handling is currently limited to windows platforms
+#if WITH_FREEIMAGE_LIB
 
 #include "GenericPlatform/GenericPlatformFile.h"
 #include "HAL/PlatformFilemanager.h"
@@ -14,7 +14,6 @@
 #include "Misc/Paths.h"
 #include "Modules/ModuleManager.h"
 
-#include "Windows/AllowWindowsPlatformTypes.h"
 THIRD_PARTY_INCLUDES_START
 	#include "FreeImage.h"
 THIRD_PARTY_INCLUDES_END
@@ -141,12 +140,12 @@ bool FDatasmithTextureResize::GetBestTextureExtension(const TCHAR* Source, FStri
 	FREE_IMAGE_FORMAT FileType = FIF_UNKNOWN;
 
 	//check the file signature and deduce its format
-	FileType = FreeImage_GetFileTypeU(Source, 0);
+	FileType = FreeImage_GetFileTypeU(TCHAR_TO_WCHAR(Source), 0);
 
 	//if still unknown, try to guess the file format from the file extension
 	if (FileType == FIF_UNKNOWN)
 	{
-		FileType = FreeImage_GetFIFFromFilenameU(Source);
+		FileType = FreeImage_GetFIFFromFilenameU(TCHAR_TO_WCHAR(Source));
 	}
 
 	//if still unknown, return failure
@@ -162,7 +161,7 @@ bool FDatasmithTextureResize::GetBestTextureExtension(const TCHAR* Source, FStri
 	}
 
 	//pointer to the image, once loaded
-	FIBITMAP* Bitmap = FreeImage_LoadU(FileType, Source, 0);
+	FIBITMAP* Bitmap = FreeImage_LoadU(FileType, TCHAR_TO_WCHAR(Source), 0);
 
 	//if the image failed to load, return failure
 	if (!Bitmap)
@@ -215,12 +214,12 @@ EDSTextureUtilsError FDatasmithTextureResize::ResizeTexture(const TCHAR* Source,
 	FREE_IMAGE_FORMAT FileType = FIF_UNKNOWN;
 
 	//check the file signature and deduce its format
-	FileType = FreeImage_GetFileTypeU(Source, 0);
+	FileType = FreeImage_GetFileTypeU(TCHAR_TO_WCHAR(Source), 0);
 
 	//if still unknown, try to guess the file format from the file extension
 	if (FileType == FIF_UNKNOWN)
 	{
-		FileType = FreeImage_GetFIFFromFilenameU(Source);
+		FileType = FreeImage_GetFIFFromFilenameU(TCHAR_TO_WCHAR(Source));
 	}
 
 	//if still unknown, return failure
@@ -236,7 +235,7 @@ EDSTextureUtilsError FDatasmithTextureResize::ResizeTexture(const TCHAR* Source,
 	}
 
 	//pointer to the image, once loaded
-	FIBITMAP* Bitmap = FreeImage_LoadU(FileType, Source, 0);
+	FIBITMAP* Bitmap = FreeImage_LoadU(FileType, TCHAR_TO_WCHAR(Source), 0);
 
 	//if the image failed to load, return failure
 	if (!Bitmap)
@@ -263,7 +262,7 @@ EDSTextureUtilsError FDatasmithTextureResize::ResizeTexture(const TCHAR* Source,
 		return EDSTextureUtilsError::InvalidData;
 	}
 
-	FREE_IMAGE_FORMAT FifW = FreeImage_GetFIFFromFilenameU(Destination);
+	FREE_IMAGE_FORMAT FifW = FreeImage_GetFIFFromFilenameU(TCHAR_TO_WCHAR(Destination));
 
 	if (bGenerateNormalMap)
 	{
@@ -310,20 +309,20 @@ EDSTextureUtilsError FDatasmithTextureResize::ResizeTexture(const TCHAR* Source,
 
 		return EDSTextureUtilsError::NoError;
 	}
-	
+
 	BOOL bSuccessfullySaved = FALSE;
 
 	if (FifW == FIF_EXR || FifW == FIF_HDR)
 	{
 		//FIBITMAP* InFloat = FreeImage_ConvertToRGBAF(RescaledImage);
-		bSuccessfullySaved = FreeImage_SaveU(FifW, RescaledImage, Destination, 0);
+		bSuccessfullySaved = FreeImage_SaveU(FifW, RescaledImage, TCHAR_TO_WCHAR(Destination), 0);
 		//FreeImage_Unload(InFloat);
 	}
 	else if (bGenerateNormalMap == true)
 	{
 		FIBITMAP* InFloat = FreeImage_ConvertTo24Bits(RescaledImage);
 
-		bSuccessfullySaved = FreeImage_SaveU(FifW, InFloat, Destination, 0);
+		bSuccessfullySaved = FreeImage_SaveU(FifW, InFloat, TCHAR_TO_WCHAR(Destination), 0);
 		FreeImage_Unload(InFloat);
 	}
 	else
@@ -331,7 +330,7 @@ EDSTextureUtilsError FDatasmithTextureResize::ResizeTexture(const TCHAR* Source,
 		// Most export format supports FIT_BITMAP image type, try it
 		if (FreeImage_GetImageType(RescaledImage) == FIT_BITMAP)
 		{
-			bSuccessfullySaved = FreeImage_SaveU(FifW, RescaledImage, Destination, 0);
+			bSuccessfullySaved = FreeImage_SaveU(FifW, RescaledImage, TCHAR_TO_WCHAR(Destination), 0);
 		}
 
 		if (!bSuccessfullySaved)
@@ -353,7 +352,7 @@ EDSTextureUtilsError FDatasmithTextureResize::ResizeTexture(const TCHAR* Source,
 					ConvertedBitmap = FreeImage_ConvertToType(RescaledImage, FIT_BITMAP, TRUE);
 				}
 
-				bSuccessfullySaved = FreeImage_SaveU(FifW, ConvertedBitmap, Destination, 0);
+				bSuccessfullySaved = FreeImage_SaveU(FifW, ConvertedBitmap, TCHAR_TO_WCHAR(Destination), 0);
 
 				if (ConvertedBitmap != RescaledImage)
 				{
@@ -364,7 +363,7 @@ EDSTextureUtilsError FDatasmithTextureResize::ResizeTexture(const TCHAR* Source,
 			{
 				ChannelCount += (BitsPerPixel % ChannelCount) ? 1 : 0;
 				FIBITMAP* In8BitPerChannel = ChannelCount > 3 ? FreeImage_ConvertTo32Bits(RescaledImage) : FreeImage_ConvertTo24Bits(RescaledImage);
-				bSuccessfullySaved = FreeImage_SaveU(FifW, In8BitPerChannel, Destination, 0);
+				bSuccessfullySaved = FreeImage_SaveU(FifW, In8BitPerChannel, TCHAR_TO_WCHAR(Destination), 0);
 				FreeImage_Unload(In8BitPerChannel);
 			}
 		}
@@ -657,8 +656,9 @@ void FFreeImageWrapper::FreeImage_Initialise(bool bLoadLocalPluginsOnly)
 	if ( FreeImageDllHandle == nullptr )
 	{
 		FString FreeImageDir = FPaths::Combine( FPaths::EngineDir(), TEXT("Binaries/ThirdParty/FreeImage"), FPlatformProcess::GetBinariesSubdirectory() );
+		FString FreeImageLibDir = FPaths::Combine( FreeImageDir, TEXT( FREEIMAGE_LIB_FILENAME ));
 		FPlatformProcess::PushDllDirectory( *FreeImageDir );
-		FreeImageDllHandle = FPlatformProcess::GetDllHandle( TEXT("FreeImage.dll") );
+		FreeImageDllHandle = FPlatformProcess::GetDllHandle( *FreeImageLibDir );
 		FPlatformProcess::PopDllDirectory( *FreeImageDir );
 	}
 
@@ -668,6 +668,4 @@ void FFreeImageWrapper::FreeImage_Initialise(bool bLoadLocalPluginsOnly)
 	}
 }
 
-#include "Windows/HideWindowsPlatformTypes.h"
-
-#endif // PLATFORM_WINDOWS
+#endif // WITH_FREEIMAGE_LIB
diff --git a/Engine/Plugins/Enterprise/DatasmithImporter/Source/DatasmithImporter/Public/Utility/DatasmithTextureResize.h b/Engine/Plugins/Enterprise/DatasmithImporter/Source/DatasmithImporter/Public/Utility/DatasmithTextureResize.h
--- a/Engine/Plugins/Enterprise/DatasmithImporter/Source/DatasmithImporter/Public/Utility/DatasmithTextureResize.h
+++ b/Engine/Plugins/Enterprise/DatasmithImporter/Source/DatasmithImporter/Public/Utility/DatasmithTextureResize.h
@@ -9,8 +9,7 @@
 class DATASMITHIMPORTER_API FDatasmithTextureResize
 {
 public:
-#if PLATFORM_WINDOWS // #ueent_todo Datasmith texture handling is currently limited to windows platforms
-
+#if WITH_FREEIMAGE_LIB
 	/**
 	 * Returns true if image extension is supported by Unreal Engine
 	 * IMPORTANT: Extension is expected with the starting dot
@@ -56,5 +55,5 @@ public:
 	*/
 	static EDSTextureUtilsError ResizeTexture( const TCHAR* Source, const TCHAR* Destination, EDSResizeTextureMode Mode, uint32 MaxSize, bool bGenerateNormalMap )
 	{ return EDSTextureUtilsError::ResizeFailed; }
-#endif // PLATFORM_WINDOWS
+#endif // WITH_FREEIMAGE_LIB
 };
