Copy of 9484814 from //UE4/Dev-Enterprise-Minimal:

Datasmith: Fix usage of freeimage on linux

diff --git a/Engine/Plugins/Enterprise/DatasmithImporter/Source/DatasmithImporter/Private/Utility/DatasmithTextureResize.cpp b/Engine/Plugins/Enterprise/DatasmithImporter/Source/DatasmithImporter/Private/Utility/DatasmithTextureResize.cpp
--- a/Engine/Plugins/Enterprise/DatasmithImporter/Source/DatasmithImporter/Private/Utility/DatasmithTextureResize.cpp
+++ b/Engine/Plugins/Enterprise/DatasmithImporter/Source/DatasmithImporter/Private/Utility/DatasmithTextureResize.cpp
@@ -18,6 +18,16 @@ THIRD_PARTY_INCLUDES_START
 	#include "FreeImage.h"
 THIRD_PARTY_INCLUDES_END
 
+#if PLATFORM_WINDOWS
+#   define TCHAR_TO_FICHAR TCHAR_TO_WCHAR
+#   define FreeImage_GetFIFFromFilename FreeImage_GetFIFFromFilenameU
+#   define FreeImage_GetFileType        FreeImage_GetFileTypeU
+#   define FreeImage_Load               FreeImage_LoadU
+#   define FreeImage_Save               FreeImage_SaveU
+#else
+#   define TCHAR_TO_FICHAR TCHAR_TO_UTF8
+#endif
+
 namespace
 {
 	int32 PreviousPowerOfTwo(int32 Reference)
@@ -140,12 +150,12 @@ bool FDatasmithTextureResize::GetBestTextureExtension(const TCHAR* Source, FStri
 	FREE_IMAGE_FORMAT FileType = FIF_UNKNOWN;
 
 	//check the file signature and deduce its format
-	FileType = FreeImage_GetFileTypeU(TCHAR_TO_WCHAR(Source), 0);
+	FileType = FreeImage_GetFileType(TCHAR_TO_FICHAR(Source), 0);
 
 	//if still unknown, try to guess the file format from the file extension
 	if (FileType == FIF_UNKNOWN)
 	{
-		FileType = FreeImage_GetFIFFromFilenameU(TCHAR_TO_WCHAR(Source));
+		FileType = FreeImage_GetFIFFromFilename(TCHAR_TO_FICHAR(Source));
 	}
 
 	//if still unknown, return failure
@@ -161,7 +171,7 @@ bool FDatasmithTextureResize::GetBestTextureExtension(const TCHAR* Source, FStri
 	}
 
 	//pointer to the image, once loaded
-	FIBITMAP* Bitmap = FreeImage_LoadU(FileType, TCHAR_TO_WCHAR(Source), 0);
+	FIBITMAP* Bitmap = FreeImage_Load(FileType, TCHAR_TO_FICHAR(Source), 0);
 
 	//if the image failed to load, return failure
 	if (!Bitmap)
@@ -214,12 +224,12 @@ EDSTextureUtilsError FDatasmithTextureResize::ResizeTexture(const TCHAR* Source,
 	FREE_IMAGE_FORMAT FileType = FIF_UNKNOWN;
 
 	//check the file signature and deduce its format
-	FileType = FreeImage_GetFileTypeU(TCHAR_TO_WCHAR(Source), 0);
+	FileType = FreeImage_GetFileType(TCHAR_TO_FICHAR(Source), 0);
 
 	//if still unknown, try to guess the file format from the file extension
 	if (FileType == FIF_UNKNOWN)
 	{
-		FileType = FreeImage_GetFIFFromFilenameU(TCHAR_TO_WCHAR(Source));
+		FileType = FreeImage_GetFIFFromFilename(TCHAR_TO_FICHAR(Source));
 	}
 
 	//if still unknown, return failure
@@ -235,7 +245,7 @@ EDSTextureUtilsError FDatasmithTextureResize::ResizeTexture(const TCHAR* Source,
 	}
 
 	//pointer to the image, once loaded
-	FIBITMAP* Bitmap = FreeImage_LoadU(FileType, TCHAR_TO_WCHAR(Source), 0);
+	FIBITMAP* Bitmap = FreeImage_Load(FileType, TCHAR_TO_FICHAR(Source), 0);
 
 	//if the image failed to load, return failure
 	if (!Bitmap)
@@ -262,7 +272,7 @@ EDSTextureUtilsError FDatasmithTextureResize::ResizeTexture(const TCHAR* Source,
 		return EDSTextureUtilsError::InvalidData;
 	}
 
-	FREE_IMAGE_FORMAT FifW = FreeImage_GetFIFFromFilenameU(TCHAR_TO_WCHAR(Destination));
+	FREE_IMAGE_FORMAT FifW = FreeImage_GetFIFFromFilename(TCHAR_TO_FICHAR(Destination));
 
 	if (bGenerateNormalMap)
 	{
@@ -315,14 +325,14 @@ EDSTextureUtilsError FDatasmithTextureResize::ResizeTexture(const TCHAR* Source,
 	if (FifW == FIF_EXR || FifW == FIF_HDR)
 	{
 		//FIBITMAP* InFloat = FreeImage_ConvertToRGBAF(RescaledImage);
-		bSuccessfullySaved = FreeImage_SaveU(FifW, RescaledImage, TCHAR_TO_WCHAR(Destination), 0);
+		bSuccessfullySaved = FreeImage_Save(FifW, RescaledImage, TCHAR_TO_FICHAR(Destination), 0);
 		//FreeImage_Unload(InFloat);
 	}
 	else if (bGenerateNormalMap == true)
 	{
 		FIBITMAP* InFloat = FreeImage_ConvertTo24Bits(RescaledImage);
 
-		bSuccessfullySaved = FreeImage_SaveU(FifW, InFloat, TCHAR_TO_WCHAR(Destination), 0);
+		bSuccessfullySaved = FreeImage_Save(FifW, InFloat, TCHAR_TO_FICHAR(Destination), 0);
 		FreeImage_Unload(InFloat);
 	}
 	else
@@ -330,7 +340,7 @@ EDSTextureUtilsError FDatasmithTextureResize::ResizeTexture(const TCHAR* Source,
 		// Most export format supports FIT_BITMAP image type, try it
 		if (FreeImage_GetImageType(RescaledImage) == FIT_BITMAP)
 		{
-			bSuccessfullySaved = FreeImage_SaveU(FifW, RescaledImage, TCHAR_TO_WCHAR(Destination), 0);
+			bSuccessfullySaved = FreeImage_Save(FifW, RescaledImage, TCHAR_TO_FICHAR(Destination), 0);
 		}
 
 		if (!bSuccessfullySaved)
@@ -352,7 +362,7 @@ EDSTextureUtilsError FDatasmithTextureResize::ResizeTexture(const TCHAR* Source,
 					ConvertedBitmap = FreeImage_ConvertToType(RescaledImage, FIT_BITMAP, TRUE);
 				}
 
-				bSuccessfullySaved = FreeImage_SaveU(FifW, ConvertedBitmap, TCHAR_TO_WCHAR(Destination), 0);
+				bSuccessfullySaved = FreeImage_Save(FifW, ConvertedBitmap, TCHAR_TO_FICHAR(Destination), 0);
 
 				if (ConvertedBitmap != RescaledImage)
 				{
@@ -363,7 +373,7 @@ EDSTextureUtilsError FDatasmithTextureResize::ResizeTexture(const TCHAR* Source,
 			{
 				ChannelCount += (BitsPerPixel % ChannelCount) ? 1 : 0;
 				FIBITMAP* In8BitPerChannel = ChannelCount > 3 ? FreeImage_ConvertTo32Bits(RescaledImage) : FreeImage_ConvertTo24Bits(RescaledImage);
-				bSuccessfullySaved = FreeImage_SaveU(FifW, In8BitPerChannel, TCHAR_TO_WCHAR(Destination), 0);
+				bSuccessfullySaved = FreeImage_Save(FifW, In8BitPerChannel, TCHAR_TO_FICHAR(Destination), 0);
 				FreeImage_Unload(In8BitPerChannel);
 			}
 		}
