Integrating CL# 9541411 from Dev-Enterprise
Adding warning on mesh with edge-case rotation values not well supported by unreal.

diff --git a/Engine/Plugins/Enterprise/DatasmithImporter/Source/DatasmithImporter/Private/DatasmithImportContext.cpp b/Engine/Plugins/Enterprise/DatasmithImporter/Source/DatasmithImporter/Private/DatasmithImportContext.cpp
--- a/Engine/Plugins/Enterprise/DatasmithImporter/Source/DatasmithImporter/Private/DatasmithImportContext.cpp
+++ b/Engine/Plugins/Enterprise/DatasmithImporter/Source/DatasmithImporter/Private/DatasmithImportContext.cpp
@@ -187,19 +187,19 @@ void FDatasmithImportContext::AddOption(UObject* InOption, bool bLoadConfig)
 	}
 }
 
-void FDatasmithImportContext::LogError(const FText& InErrorMessage)
+TSharedRef<FTokenizedMessage> FDatasmithImportContext::LogError(const FText& InErrorMessage)
 {
-	Logger.Push(EMessageSeverity::Error, InErrorMessage);
+	return Logger.Push(EMessageSeverity::Error, InErrorMessage);
 }
 
-void FDatasmithImportContext::LogWarning(const FText& InWarningMessage, bool bPerformance)
+TSharedRef<FTokenizedMessage> FDatasmithImportContext::LogWarning(const FText& InWarningMessage, bool bPerformance)
 {
-	Logger.Push(bPerformance ? EMessageSeverity::PerformanceWarning : EMessageSeverity::Warning, InWarningMessage);
+	return Logger.Push(bPerformance ? EMessageSeverity::PerformanceWarning : EMessageSeverity::Warning, InWarningMessage);
 }
 
-void FDatasmithImportContext::LogInfo(const FText& InInfoMessage)
+TSharedRef<FTokenizedMessage> FDatasmithImportContext::LogInfo(const FText& InInfoMessage)
 {
-	Logger.Push(EMessageSeverity::Info, InInfoMessage);
+	return Logger.Push(EMessageSeverity::Info, InInfoMessage);
 }
 
 bool FDatasmithImportContext::Init(const FString& InFileName, TSharedRef< IDatasmithScene > InScene, const FString& InImportPath, EObjectFlags InFlags, FFeedbackContext* InWarn, const TSharedPtr<FJsonObject>& ImportSettingsJson, bool bSilent)
diff --git a/Engine/Plugins/Enterprise/DatasmithImporter/Source/DatasmithImporter/Private/DatasmithImporter.cpp b/Engine/Plugins/Enterprise/DatasmithImporter/Source/DatasmithImporter/Private/DatasmithImporter.cpp
--- a/Engine/Plugins/Enterprise/DatasmithImporter/Source/DatasmithImporter/Private/DatasmithImporter.cpp
+++ b/Engine/Plugins/Enterprise/DatasmithImporter/Source/DatasmithImporter/Private/DatasmithImporter.cpp
@@ -56,6 +56,7 @@
 #include "Misc/FeedbackContext.h"
 #include "Misc/FileHelper.h"
 #include "Misc/ScopedSlowTask.h"
+#include "Misc/UObjectToken.h"
 #include "ObjectTools.h"
 #include "PackageTools.h"
 #include "Serialization/ArchiveReplaceObjectRef.h"
@@ -1793,7 +1794,31 @@ AActor* FDatasmithImporter::FinalizeActor(FDatasmithImportContext& ImportContext
 		Landscape->PostEditChangeProperty(MaterialPropertyChangedEvent);
 	}
 
+	FQuat PreviousRotation = DestinationActor->GetRootComponent()->GetRelativeTransform().GetRotation();
 	DestinationActor->PostEditChange();
+
+	const bool bHasPostEditChangeModifiedRotation = !PreviousRotation.Equals(DestinationActor->GetRootComponent()->GetRelativeTransform().GetRotation());
+	if (bHasPostEditChangeModifiedRotation)
+	{
+		const float SingularityTest = PreviousRotation.Z * PreviousRotation.X - PreviousRotation.W * PreviousRotation.Y;
+		//SingularityThreshold value is comming from the FQuat::Rotator() function, but is more permissive because the rotation is already diverging before the singularity threshold is reached.
+		const float SingularityThreshold = 0.4999f; 
+
+		AActor* RootSceneActor = ImportContext.ActorsContext.ImportSceneActor;
+		if (DestinationActor != RootSceneActor
+			&& FMath::Abs(SingularityTest) > SingularityThreshold)
+		{
+			//This is a warning to explain the edge-case of UE-75467 while it's being fixed.
+			FFormatNamedArguments FormatArgs;
+			FormatArgs.Add(TEXT("ActorName"), FText::FromName(DestinationActor->GetFName()));
+			ImportContext.LogWarning(FText::GetEmpty())
+				->AddToken(FUObjectToken::Create(DestinationActor))
+				->AddToken(FTextToken::Create(FText::Format(LOCTEXT("UnsupportedRotationValueError", "The actor '{ActorName}' has a rotation value pointing to either (0, 90, 0) or (0, -90, 0)."
+					"This is an edge case that is not well supported in Unreal and can cause incorrect results."
+					"In those cases, it is recommended to bake the actor's transform into the mesh at export."), FormatArgs)));
+		}
+	}
+
 	DestinationActor->RegisterAllComponents();
 
 	return DestinationActor;
diff --git a/Engine/Plugins/Enterprise/DatasmithImporter/Source/DatasmithImporter/Private/Utility/DatasmithImporterUtils.cpp b/Engine/Plugins/Enterprise/DatasmithImporter/Source/DatasmithImporter/Private/Utility/DatasmithImporterUtils.cpp
--- a/Engine/Plugins/Enterprise/DatasmithImporter/Source/DatasmithImporter/Private/Utility/DatasmithImporterUtils.cpp
+++ b/Engine/Plugins/Enterprise/DatasmithImporter/Source/DatasmithImporter/Private/Utility/DatasmithImporterUtils.cpp
@@ -983,9 +983,11 @@ FScopedLogger::~FScopedLogger()
 	Dump(true);
 }
 
-void FScopedLogger::Push(EMessageSeverity::Type Severity, const FText& Message)
+TSharedRef<FTokenizedMessage> FScopedLogger::Push(EMessageSeverity::Type Severity, const FText& Message)
 {
 	TokenizedMessages.Add(FTokenizedMessage::Create(Severity, Message));
+	
+	return TokenizedMessages.Last();
 }
 
 void FScopedLogger::Dump(bool bClearPrevious)
diff --git a/Engine/Plugins/Enterprise/DatasmithImporter/Source/DatasmithImporter/Public/DatasmithImportContext.h b/Engine/Plugins/Enterprise/DatasmithImporter/Source/DatasmithImporter/Public/DatasmithImportContext.h
--- a/Engine/Plugins/Enterprise/DatasmithImporter/Source/DatasmithImporter/Public/DatasmithImportContext.h
+++ b/Engine/Plugins/Enterprise/DatasmithImporter/Source/DatasmithImporter/Public/DatasmithImportContext.h
@@ -72,7 +72,7 @@ public:
 	FScopedLogger& operator = (FScopedLogger&&) = delete;
 
 	/** Append a message */
-	void Push(EMessageSeverity::Type Severity, const FText& Message);
+	TSharedRef<FTokenizedMessage> Push(EMessageSeverity::Type Severity, const FText& Message);
 
 	/** Display pending messages, show the log window */
 	void Dump(bool bClearPrevious = false);
@@ -301,9 +301,9 @@ public:
 	void AddOption(UObject* InOption, bool bLoadConfig);
 
 	/** Push messages for display to the end-user */
-	void LogError(const FText& InErrorMessage);
-	void LogWarning(const FText& InWarningMessage, bool bPerformance = false);
-	void LogInfo(const FText& InInfoMessage);
+	TSharedRef<FTokenizedMessage> LogError(const FText& InErrorMessage);
+	TSharedRef<FTokenizedMessage> LogWarning(const FText& InWarningMessage, bool bPerformance = false);
+	TSharedRef<FTokenizedMessage> LogInfo(const FText& InInfoMessage);
 
 	bool ShouldImportActors() const;
 
