Fixed naming issue with file producer

diff --git a/Engine/Plugins/Enterprise/DatasmithImporter/Source/DatasmithImporter/Private/DatasmithFileProducer.cpp b/Engine/Plugins/Enterprise/DatasmithImporter/Source/DatasmithImporter/Private/DatasmithFileProducer.cpp
--- a/Engine/Plugins/Enterprise/DatasmithImporter/Source/DatasmithImporter/Private/DatasmithFileProducer.cpp
+++ b/Engine/Plugins/Enterprise/DatasmithImporter/Source/DatasmithImporter/Private/DatasmithFileProducer.cpp
@@ -563,27 +563,24 @@ void UDatasmithFileProducer::SetFilename( const FString& InFilename )
 
 	FilePath = FPaths::ConvertRelativePathToFull( InFilename );
 
-	// Rename producer to name of file
-	FString CleanName = ObjectTools::SanitizeObjectName( FPaths::GetCleanFilename( InFilename ) );
-	if ( Rename( *CleanName, nullptr, REN_Test ) )
-	{
-		Rename( *CleanName, nullptr, REN_DontCreateRedirectors | REN_NonTransactional );
-	}
-	else
+	UpdateName();
+
+	OnChanged.Broadcast( this );
+}
+
+void UDatasmithFileProducer::UpdateName()
+{
+	if(!FilePath.IsEmpty())
 	{
-		bool bFoundName = false;
-		for (int32 NameIndex = 1; !bFoundName; ++NameIndex)
+		// Rename producer to name of file
+		FString CleanName = ObjectTools::SanitizeObjectName( FPaths::GetCleanFilename( FilePath ) );
+		if ( !Rename( *CleanName, nullptr, REN_Test ) )
 		{
-			const FString NewName = FString::Printf(TEXT("%s_%d"), *CleanName, NameIndex);
-			if( Rename( *NewName, nullptr, REN_Test ) )
-			{
-				Rename( *NewName, nullptr, REN_DontCreateRedirectors | REN_NonTransactional );
-				bFoundName = true;
-			}
+			CleanName = MakeUniqueObjectName( GetOuter(), GetClass(), *CleanName ).ToString();
 		}
-	}
 
-	OnChanged.Broadcast( this );
+		Rename( *CleanName, nullptr, REN_DontCreateRedirectors | REN_NonTransactional );
+	}
 }
 
 bool UDatasmithFileProducer::Supersede(const UDataprepContentProducer* OtherProducer) const
@@ -610,6 +607,7 @@ void UDatasmithFileProducer::PostInitProperties()
 	if( !HasAnyFlags( RF_ClassDefaultObject | RF_WasLoaded | RF_Transient ) )
 	{
 		FilePath = FDatasmithFileProducerUtils::SelectFileToImport();
+		UpdateName();
 	}
 }
 
diff --git a/Engine/Plugins/Enterprise/DatasmithImporter/Source/DatasmithImporter/Private/DatasmithFileProducer.h b/Engine/Plugins/Enterprise/DatasmithImporter/Source/DatasmithImporter/Private/DatasmithFileProducer.h
--- a/Engine/Plugins/Enterprise/DatasmithImporter/Source/DatasmithImporter/Private/DatasmithFileProducer.h
+++ b/Engine/Plugins/Enterprise/DatasmithImporter/Source/DatasmithImporter/Private/DatasmithFileProducer.h
@@ -61,6 +61,9 @@ private:
 	/** Fill up world with content of Datasmith scene element */
 	void PreventNameCollision();
 
+	/** Update the name of the producer based on the filename */
+	void UpdateName();
+
 private:
 	TUniquePtr< FDatasmithImportContext > ImportContextPtr;
 	TUniquePtr< FDatasmithTranslatableSceneSource > TranslatableSourcePtr;
