Restore absolute minimum lightmap computation

From dev-enterprise[REDACTED_HANDLE]

diff --git a/Engine/Plugins/Enterprise/DatasmithImporter/Source/DatasmithImporter/Private/DatasmithStaticMeshImporter.cpp b/Engine/Plugins/Enterprise/DatasmithImporter/Source/DatasmithImporter/Private/DatasmithStaticMeshImporter.cpp
--- a/Engine/Plugins/Enterprise/DatasmithImporter/Source/DatasmithImporter/Private/DatasmithStaticMeshImporter.cpp
+++ b/Engine/Plugins/Enterprise/DatasmithImporter/Source/DatasmithImporter/Private/DatasmithStaticMeshImporter.cpp
@@ -214,6 +214,18 @@ bool FDatasmithStaticMeshImporter::PreBuildStaticMesh( UStaticMesh* StaticMesh )
 			if (BuildSettings.bGenerateLightmapUVs)
 			{
 				DatasmithMeshHelper::RequireUVChannel(MeshDescription, BuildSettings.DstLightmapIndex);
+
+				// Determine the absolute minimum lightmap resolution that can be used for packing
+				float ComparisonThreshold = BuildSettings.bRemoveDegenerates ? THRESH_POINTS_ARE_SAME : 0.0f;
+				FOverlappingCorners OverlappingCorners;
+				FMeshDescriptionOperations::FindOverlappingCorners(OverlappingCorners, MeshDescription, ComparisonThreshold);
+
+				// Packing expects at least one texel per chart. This is the absolute minimum to generate valid UVs.
+				int32 ChartCount = FMeshDescriptionOperations::GetUVChartCount(MeshDescription, BuildSettings.SrcLightmapIndex, ELightmapUVVersion::Latest, OverlappingCorners);
+				const int32 AbsoluteMinResolution = 1 << FMath::CeilLogTwo(FMath::Sqrt(ChartCount));
+				const int32 LightmapResolution = FMath::Clamp(BuildSettings.MinLightmapResolution, AbsoluteMinResolution, 512);
+
+				BuildSettings.MinLightmapResolution = LightmapResolution;
 			}
 
 			UStaticMesh::FCommitMeshDescriptionParams Params;
