Parallel texture loading
  - 47s -> 16 s for DATASET-0008a (textures only)

From dev-enterprise[REDACTED_HANDLE],9527316,9538821

diff --git a/Engine/Plugins/Enterprise/DatasmithImporter/Source/DatasmithImporter/Private/DatasmithImporter.cpp b/Engine/Plugins/Enterprise/DatasmithImporter/Source/DatasmithImporter/Private/DatasmithImporter.cpp
--- a/Engine/Plugins/Enterprise/DatasmithImporter/Source/DatasmithImporter/Private/DatasmithImporter.cpp
+++ b/Engine/Plugins/Enterprise/DatasmithImporter/Source/DatasmithImporter/Private/DatasmithImporter.cpp
@@ -25,6 +25,7 @@
 #include "Translators/DatasmithPayload.h"
 #include "Translators/DatasmithTranslator.h"
 #include "Utility/DatasmithImporterUtils.h"
+#include "Utility/DatasmithTextureResize.h"
 
 #include "AssetRegistryModule.h"
 #include "AssetToolsModule.h"
@@ -86,7 +87,7 @@ namespace DatasmithImporterImpl
 		}
 	}
 
-	static UObject* PublicizeAsset(UObject* SourceAsset, const TCHAR* DestinationPath, UObject* ExistingAsset)
+	static UObject* PublicizeAsset( UObject* SourceAsset, const TCHAR* DestinationPath, UObject* ExistingAsset )
 	{
 		TRACE_CPUPROFILER_EVENT_SCOPE(DatasmithImporterImpl::PublicizeAsset);
 
@@ -94,14 +95,14 @@ namespace DatasmithImporterImpl
 
 		UPackage* DestinationPackage;
 
-		if (!ExistingAsset)
+		if ( !ExistingAsset )
 		{
-			FString DestinationPackagePath = UPackageTools::SanitizePackageName(FPaths::Combine(DestinationPath, SourceAsset->GetName()));
-			FString DestinationAssetPath = DestinationPackagePath + TEXT(".") + UPackageTools::SanitizePackageName(SourceAsset->GetName());
+			FString DestinationPackagePath = UPackageTools::SanitizePackageName( FPaths::Combine( DestinationPath, SourceAsset->GetName() ) );
+			FString DestinationAssetPath = DestinationPackagePath + TEXT(".") + UPackageTools::SanitizePackageName( SourceAsset->GetName() );
 
-			ExistingAsset = FDatasmithImporterUtils::FindObject<UObject>(nullptr, DestinationAssetPath);
+			ExistingAsset = FDatasmithImporterUtils::FindObject<UObject>( nullptr, DestinationAssetPath );
 
-			DestinationPackage = ExistingAsset ? ExistingAsset->GetOutermost() : CreatePackage(nullptr, *DestinationPackagePath);
+			DestinationPackage = ExistingAsset ? ExistingAsset->GetOutermost() : CreatePackage( nullptr, *DestinationPackagePath );
 		}
 		else
 		{
@@ -115,26 +116,26 @@ namespace DatasmithImporterImpl
 		FString OldAssetPathName;
 
 		// If the object already exist, then we need to fix up the reference
-		if (ExistingAsset != nullptr && ExistingAsset != SourceAsset)
+		if ( ExistingAsset != nullptr && ExistingAsset != SourceAsset )
 		{
 			OldAssetPathName = ExistingAsset->GetPathName();
 
-			DestinationAsset = DuplicateObject< UObject >(SourceAsset, DestinationPackage, ExistingAsset->GetFName());
+			DestinationAsset = DuplicateObject< UObject >( SourceAsset, DestinationPackage, ExistingAsset->GetFName() );
 
 			// If mesh's label has changed, update its name
-			if (ExistingAsset->GetFName() != SourceAsset->GetFName())
+			if ( ExistingAsset->GetFName() != SourceAsset->GetFName() )
 			{
-				DestinationAsset->Rename(*SourceAsset->GetName(), DestinationPackage, REN_DontCreateRedirectors | REN_NonTransactional);
+				DestinationAsset->Rename( *SourceAsset->GetName(), DestinationPackage, REN_DontCreateRedirectors | REN_NonTransactional );
 			}
 
-			if (UStaticMesh* DestinationMesh = Cast< UStaticMesh >(DestinationAsset))
+			if ( UStaticMesh* DestinationMesh = Cast< UStaticMesh >( DestinationAsset ) )
 			{
 				// This is done during the static mesh build process but we need to redo it after the DuplicateObject since the links are now valid
-				for (TObjectIterator< UStaticMeshComponent > It; It; ++It)
+				for ( TObjectIterator< UStaticMeshComponent > It; It; ++It )
 				{
-					if (It->GetStaticMesh() == DestinationMesh)
+					if ( It->GetStaticMesh() == DestinationMesh )
 					{
-						It->FixupOverrideColorsIfNecessary(true);
+						It->FixupOverrideColorsIfNecessary( true );
 						It->InvalidateLightingCache();
 					}
 				}
@@ -142,20 +143,20 @@ namespace DatasmithImporterImpl
 		}
 		else
 		{
-			SourceAsset->Rename(*SourceAsset->GetName(), DestinationPackage, REN_DontCreateRedirectors | REN_NonTransactional);
+			SourceAsset->Rename( *SourceAsset->GetName(), DestinationPackage, REN_DontCreateRedirectors | REN_NonTransactional );
 			DestinationAsset = SourceAsset;
 		}
 
-		DestinationAsset->SetFlags(RF_Public);
+		DestinationAsset->SetFlags( RF_Public );
 		DestinationAsset->MarkPackageDirty();
 
-		if (!ExistingAsset)
+		if ( !ExistingAsset )
 		{
-			FAssetRegistryModule::AssetCreated(DestinationAsset);
+			FAssetRegistryModule::AssetCreated( DestinationAsset );
 		}
-		else if (!OldAssetPathName.IsEmpty())
+		else if ( !OldAssetPathName.IsEmpty() )
 		{
-			FAssetRegistryModule::AssetRenamed(DestinationAsset, OldAssetPathName);
+			FAssetRegistryModule::AssetRenamed( DestinationAsset, OldAssetPathName );
 		}
 
 		return DestinationAsset;
@@ -178,22 +179,22 @@ namespace DatasmithImporterImpl
 		const FString PackageName = Package->GetName();
 
 		// Check that package can be saved
-		const FString BasePackageFileName = FPackageName::LongPackageNameToFilename(PackageName);
-		const FString AbsolutePathToAsset = FPaths::ConvertRelativePathToFull(BasePackageFileName);
+		const FString BasePackageFileName = FPackageName::LongPackageNameToFilename( PackageName );
+		const FString AbsolutePathToAsset = FPaths::ConvertRelativePathToFull( BasePackageFileName );
 
 		// Create fake filename of same length of final asset file name to test ability to write
-		const FString FakeAbsolutePathToAsset = AbsolutePathToAsset + TEXT(".uasset");
+		const FString FakeAbsolutePathToAsset = AbsolutePathToAsset + TEXT( ".uasset" );
 
 		// Verify asset file name does not exceed OS' maximum path length
-		if (FPlatformMisc::GetMaxPathLength() < FakeAbsolutePathToAsset.Len())
+		if( FPlatformMisc::GetMaxPathLength() < FakeAbsolutePathToAsset.Len() )
 		{
-			ImportContext.LogWarning(FText::Format(LOCTEXT("DatasmithImportInvalidLength", "Saving may partially fail because path for asset {0} is too long. Rename before saving."), FText::FromString(PackageName)));
+			ImportContext.LogWarning(FText::Format(LOCTEXT("DatasmithImportInvalidLength", "Saving may partially fail because path for asset {0} is too long. Rename before saving."), FText::FromString( PackageName )));
 		}
 		// Verify user can overwrite existing file
-		else if (IFileManager::Get().FileExists(*FakeAbsolutePathToAsset))
+		else if( IFileManager::Get().FileExists( *FakeAbsolutePathToAsset ) )
 		{
-			FFileStatData FileStatData = IFileManager::Get().GetStatData(*FakeAbsolutePathToAsset);
-			if (FileStatData.bIsReadOnly)
+			FFileStatData FileStatData = IFileManager::Get().GetStatData( *FakeAbsolutePathToAsset );
+			if ( FileStatData.bIsReadOnly )
 			{
 				// Check to see if the file is not under source control
 				bool bWarnUser = true;
@@ -203,16 +204,16 @@ namespace DatasmithImporterImpl
 				{
 					SourceControlProvider.Execute(ISourceControlOperation::Create<FUpdateStatus>(), FakeAbsolutePathToAsset);
 					FSourceControlStatePtr SourceControlState = SourceControlProvider.GetState(FakeAbsolutePathToAsset, EStateCacheUsage::Use);
-					if (SourceControlState.IsValid() && SourceControlState->CanCheckout())
+					if( SourceControlState.IsValid() && SourceControlState->CanCheckout() )
 					{
 						// User will be prompted to check out this file when he/she saves the asset. No need to warn.
 						bWarnUser = false;
 					}
 				}
 
-				if (bWarnUser)
+				if(bWarnUser)
 				{
-					ImportContext.LogWarning(FText::Format(LOCTEXT("DatasmithImportInvalidSaving", "Saving may partially fail because file asset {0} cannot be overwritten. Check your privileges."), FText::FromString(PackageName)));
+					ImportContext.LogWarning(FText::Format(LOCTEXT("DatasmithImportInvalidSaving", "Saving may partially fail because file asset {0} cannot be overwritten. Check your privileges."), FText::FromString( PackageName )));
 				}
 			}
 		}
@@ -237,9 +238,9 @@ namespace DatasmithImporterImpl
 				WarningState = EWriteDisallowedWarningState::WDWS_WarningUnnecessary;
 			}
 
-			if (WarningState != EWriteDisallowedWarningState::WDWS_WarningUnnecessary)
+			if(WarningState != EWriteDisallowedWarningState::WDWS_WarningUnnecessary)
 			{
-				ImportContext.LogWarning(FText::Format(LOCTEXT("DatasmithImportInvalidFolder", "Cannot write in folder {0} to store asset {1}. Check access to folder."), FText::FromString(FPaths::GetPath(FakeAbsolutePathToAsset)), FText::FromString(PackageName)));
+				ImportContext.LogWarning(FText::Format(LOCTEXT("DatasmithImportInvalidFolder", "Cannot write in folder {0} to store asset {1}. Check access to folder."), FText::FromString( FPaths::GetPath( FakeAbsolutePathToAsset ) ), FText::FromString( PackageName )));
 			}
 		}
 
@@ -277,43 +278,43 @@ namespace DatasmithImporterImpl
 	}
 
 	/** Set the texture mode on each texture element based on its usage in the materials */
-	static void SetTexturesMode(FDatasmithImportContext& ImportContext)
+	static void SetTexturesMode( FDatasmithImportContext& ImportContext )
 	{
 		const int32 TexturesCount = ImportContext.FilteredScene->GetTexturesCount();
 		const int32 MaterialsCount = ImportContext.FilteredScene->GetMaterialsCount();
 
-		for (int32 TextureIndex = 0; TextureIndex < TexturesCount && !ImportContext.bUserCancelled; ++TextureIndex)
+		for ( int32 TextureIndex = 0; TextureIndex < TexturesCount && !ImportContext.bUserCancelled; ++TextureIndex )
 		{
 			ImportContext.bUserCancelled |= ImportContext.Warn->ReceivedUserCancel();
 
-			TSharedPtr< IDatasmithTextureElement > TextureElement = ImportContext.FilteredScene->GetTexture(TextureIndex);
-			const FString TextureName = ObjectTools::SanitizeObjectName(TextureElement->GetName());
+			TSharedPtr< IDatasmithTextureElement > TextureElement = ImportContext.FilteredScene->GetTexture( TextureIndex );
+			const FString TextureName = ObjectTools::SanitizeObjectName( TextureElement->GetName() );
 
-			for (int32 MaterialIndex = 0; MaterialIndex < MaterialsCount; ++MaterialIndex)
+			for ( int32 MaterialIndex = 0; MaterialIndex < MaterialsCount; ++MaterialIndex )
 			{
-				const TSharedPtr< IDatasmithBaseMaterialElement >& BaseMaterialElement = ImportContext.FilteredScene->GetMaterial(MaterialIndex);
+				const TSharedPtr< IDatasmithBaseMaterialElement >& BaseMaterialElement = ImportContext.FilteredScene->GetMaterial( MaterialIndex );
 
-				if (BaseMaterialElement->IsA(EDatasmithElementType::Material))
+				if ( BaseMaterialElement->IsA( EDatasmithElementType::Material ) )
 				{
-					const TSharedPtr< IDatasmithMaterialElement >& MaterialElement = StaticCastSharedPtr< IDatasmithMaterialElement >(BaseMaterialElement);
+					const TSharedPtr< IDatasmithMaterialElement >& MaterialElement = StaticCastSharedPtr< IDatasmithMaterialElement >( BaseMaterialElement );
 
-					for (int32 s = 0; s < MaterialElement->GetShadersCount(); ++s)
+					for ( int32 s = 0; s < MaterialElement->GetShadersCount(); ++s )
 					{
 						const TSharedPtr< IDatasmithShaderElement >& ShaderElement = MaterialElement->GetShader(s);
 
-						if (FCString::Strlen(ShaderElement->GetDiffuseTexture()) > 0 && ShaderElement->GetDiffuseTexture() == TextureName)
+						if ( FCString::Strlen( ShaderElement->GetDiffuseTexture() ) > 0 && ShaderElement->GetDiffuseTexture() == TextureName)
 						{
 							TextureElement->SetTextureMode(EDatasmithTextureMode::Diffuse);
 						}
-						else if (FCString::Strlen(ShaderElement->GetReflectanceTexture()) > 0 && ShaderElement->GetReflectanceTexture() == TextureName)
+						else if ( FCString::Strlen( ShaderElement->GetReflectanceTexture() ) > 0 && ShaderElement->GetReflectanceTexture() == TextureName)
 						{
 							TextureElement->SetTextureMode(EDatasmithTextureMode::Specular);
 						}
-						else if (FCString::Strlen(ShaderElement->GetDisplaceTexture()) > 0 && ShaderElement->GetDisplaceTexture() == TextureName)
+						else if ( FCString::Strlen( ShaderElement->GetDisplaceTexture() ) > 0 && ShaderElement->GetDisplaceTexture() == TextureName)
 						{
 							TextureElement->SetTextureMode(EDatasmithTextureMode::Displace);
 						}
-						else if (FCString::Strlen(ShaderElement->GetNormalTexture()) > 0 && ShaderElement->GetNormalTexture() == TextureName)
+						else if ( FCString::Strlen( ShaderElement->GetNormalTexture() ) > 0 && ShaderElement->GetNormalTexture() == TextureName)
 						{
 							if (!ShaderElement->GetNormalTextureSampler().bInvert)
 							{
@@ -326,33 +327,33 @@ namespace DatasmithImporterImpl
 						}
 					}
 				}
-				else if (BaseMaterialElement->IsA(EDatasmithElementType::UEPbrMaterial))
+				else if ( BaseMaterialElement->IsA( EDatasmithElementType::UEPbrMaterial ) )
 				{
-					const TSharedPtr< IDatasmithUEPbrMaterialElement >& MaterialElement = StaticCastSharedPtr< IDatasmithUEPbrMaterialElement >(BaseMaterialElement);
+					const TSharedPtr< IDatasmithUEPbrMaterialElement >& MaterialElement = StaticCastSharedPtr< IDatasmithUEPbrMaterialElement >( BaseMaterialElement );
 
-					TFunction< bool(IDatasmithMaterialExpression*) > IsTextureConnected;
-					IsTextureConnected = [&TextureName, &IsTextureConnected](IDatasmithMaterialExpression* MaterialExpression) -> bool
+					TFunction< bool( IDatasmithMaterialExpression* ) > IsTextureConnected;
+					IsTextureConnected = [ &TextureName, &IsTextureConnected ]( IDatasmithMaterialExpression* MaterialExpression ) -> bool
 					{
-						if (!MaterialExpression)
+						if ( !MaterialExpression )
 						{
 							return false;
 						}
 
-						if (MaterialExpression->IsA(EDatasmithMaterialExpressionType::Texture))
+						if ( MaterialExpression->IsA( EDatasmithMaterialExpressionType::Texture ) )
 						{
-							IDatasmithMaterialExpressionTexture* TextureExpression = static_cast<IDatasmithMaterialExpressionTexture*>(MaterialExpression);
+							IDatasmithMaterialExpressionTexture* TextureExpression = static_cast< IDatasmithMaterialExpressionTexture* >( MaterialExpression );
 
-							if (TextureExpression->GetTexturePathName() == TextureName)
+							if ( TextureExpression->GetTexturePathName() == TextureName )
 							{
 								return true;
 							}
 						}
 
-						for (int32 InputIndex = 0; InputIndex < MaterialExpression->GetInputCount(); ++InputIndex)
+						for ( int32 InputIndex = 0; InputIndex < MaterialExpression->GetInputCount(); ++InputIndex )
 						{
-							IDatasmithMaterialExpression* ConnectedExpression = MaterialExpression->GetInput(InputIndex)->GetExpression();
+							IDatasmithMaterialExpression* ConnectedExpression = MaterialExpression->GetInput( InputIndex )->GetExpression();
 
-							if (ConnectedExpression && IsTextureConnected(ConnectedExpression))
+							if ( ConnectedExpression && IsTextureConnected( ConnectedExpression ) )
 							{
 								return true;
 							}
@@ -361,17 +362,17 @@ namespace DatasmithImporterImpl
 						return false;
 					};
 
-					if (IsTextureConnected(MaterialElement->GetBaseColor().GetExpression()))
+					if ( IsTextureConnected( MaterialElement->GetBaseColor().GetExpression() ) )
 					{
 						TextureElement->SetTextureMode(EDatasmithTextureMode::Diffuse);
 					}
-					else if (IsTextureConnected(MaterialElement->GetSpecular().GetExpression()))
+					else if ( IsTextureConnected( MaterialElement->GetSpecular().GetExpression() ) )
 					{
 						TextureElement->SetTextureMode(EDatasmithTextureMode::Specular);
 					}
-					else if (IsTextureConnected(MaterialElement->GetNormal().GetExpression()))
+					else if ( IsTextureConnected( MaterialElement->GetNormal().GetExpression() ) )
 					{
-						if (TextureElement->GetTextureMode() != EDatasmithTextureMode::Bump)
+						if ( TextureElement->GetTextureMode() != EDatasmithTextureMode::Bump )
 						{
 							TextureElement->SetTextureMode(EDatasmithTextureMode::Normal);
 						}
@@ -381,9 +382,9 @@ namespace DatasmithImporterImpl
 		}
 	}
 
-	static void CompileMaterial(UObject* Material)
+	static void CompileMaterial( UObject* Material )
 	{
-		if (!Material->IsA< UMaterialInterface >() && !Material->IsA< UMaterialFunctionInterface >())
+		if ( !Material->IsA< UMaterialInterface >() && !Material->IsA< UMaterialFunctionInterface >() )
 		{
 			return;
 		}
@@ -392,15 +393,15 @@ namespace DatasmithImporterImpl
 
 		FMaterialUpdateContext MaterialUpdateContext;
 
-		if (UMaterialInterface* MaterialInterface = Cast< UMaterialInterface >(Material))
+		if ( UMaterialInterface* MaterialInterface = Cast< UMaterialInterface >( Material ) )
 		{
-			MaterialUpdateContext.AddMaterialInterface(MaterialInterface);
+			MaterialUpdateContext.AddMaterialInterface( MaterialInterface );
 		}
 
-		if (UMaterialInstanceConstant* ConstantMaterialInstance = Cast< UMaterialInstanceConstant >(Material))
+		if ( UMaterialInstanceConstant* ConstantMaterialInstance = Cast< UMaterialInstanceConstant >( Material ) )
 		{
 			// If BlendMode override property has been changed, make sure this combination of the parent material is compiled
-			if (ConstantMaterialInstance->BasePropertyOverrides.bOverride_BlendMode == true)
+			if ( ConstantMaterialInstance->BasePropertyOverrides.bOverride_BlendMode == true )
 			{
 				ConstantMaterialInstance->ForceRecompileForRendering();
 			}
@@ -408,11 +409,11 @@ namespace DatasmithImporterImpl
 			{
 				// If a static switch is overriden, we need to recompile
 				FStaticParameterSet StaticParameters;
-				ConstantMaterialInstance->GetStaticParameterValues(StaticParameters);
+				ConstantMaterialInstance->GetStaticParameterValues( StaticParameters );
 
-				for (FStaticSwitchParameter& Switch : StaticParameters.StaticSwitchParameters)
+				for ( FStaticSwitchParameter& Switch : StaticParameters.StaticSwitchParameters )
 				{
-					if (Switch.bOverride)
+					if ( Switch.bOverride )
 					{
 						ConstantMaterialInstance->ForceRecompileForRendering();
 						break;
@@ -421,19 +422,19 @@ namespace DatasmithImporterImpl
 			}
 		}
 
-		Material->PreEditChange(nullptr);
+		Material->PreEditChange( nullptr );
 		Material->PostEditChange();
 	}
 
-	static void FixReferencesForObject(UObject* Object, const TMap< UObject*, UObject* >& ReferencesToRemap)
+	static void FixReferencesForObject( UObject* Object, const TMap< UObject*, UObject* >& ReferencesToRemap )
 	{
 		constexpr bool bNullPrivateRefs = false;
 		constexpr bool bIgnoreOuterRef = true;
 		constexpr bool bIgnoreArchetypeRef = true;
 
-		if (ReferencesToRemap.Num() > 0)
+		if ( ReferencesToRemap.Num() > 0 )
 		{
-			FArchiveReplaceObjectRef< UObject > ArchiveReplaceObjectRef(Object, ReferencesToRemap, bNullPrivateRefs, bIgnoreOuterRef, bIgnoreArchetypeRef);
+			FArchiveReplaceObjectRef< UObject > ArchiveReplaceObjectRef( Object, ReferencesToRemap, bNullPrivateRefs, bIgnoreOuterRef, bIgnoreArchetypeRef );
 		}
 	}
 
@@ -445,25 +446,25 @@ namespace DatasmithImporterImpl
 	 * @returns An array of template pairs. The key is the template for the object, the value is a template to force apply to the object,
 	 *			it contains the values from the key and any overrides that were present on the DestinationObject.
 	 */
-	static TArray< FMigratedTemplatePairType > MigrateTemplates(UObject* SourceObject, UObject* DestinationObject, const TMap< UObject*, UObject* >* ReferencesToRemap, bool bIsForActor)
+	static TArray< FMigratedTemplatePairType > MigrateTemplates( UObject* SourceObject, UObject* DestinationObject, const TMap< UObject*, UObject* >* ReferencesToRemap, bool bIsForActor )
 	{
 		TArray< FMigratedTemplatePairType > Results;
 
-		if (!SourceObject)
+		if ( !SourceObject )
 		{
 			return Results;
 		}
 
-		TMap< TSubclassOf< UDatasmithObjectTemplate >, UDatasmithObjectTemplate* >* SourceTemplates = FDatasmithObjectTemplateUtils::FindOrCreateObjectTemplates(SourceObject);
+		TMap< TSubclassOf< UDatasmithObjectTemplate >, UDatasmithObjectTemplate* >* SourceTemplates = FDatasmithObjectTemplateUtils::FindOrCreateObjectTemplates( SourceObject );
 
-		if (!SourceTemplates)
+		if ( !SourceTemplates )
 		{
 			return Results;
 		}
 
-		for (const TPair< TSubclassOf< UDatasmithObjectTemplate >, UDatasmithObjectTemplate* >& SourceTemplatePair : *SourceTemplates)
+		for ( const TPair< TSubclassOf< UDatasmithObjectTemplate >, UDatasmithObjectTemplate* >& SourceTemplatePair : *SourceTemplates )
 		{
-			if (bIsForActor == SourceTemplatePair.Value->bIsActorTemplate)
+			if ( bIsForActor == SourceTemplatePair.Value->bIsActorTemplate )
 			{
 				FMigratedTemplatePairType& Result = Results.AddDefaulted_GetRef();
 
@@ -471,7 +472,7 @@ namespace DatasmithImporterImpl
 
 				SourceTemplate->Load(SourceObject);
 
-				if (ReferencesToRemap)
+				if ( ReferencesToRemap )
 				{
 					FixReferencesForObject(SourceTemplate.Get(), *ReferencesToRemap);
 				}
@@ -480,7 +481,7 @@ namespace DatasmithImporterImpl
 
 				if (DestinationObject && !DestinationObject->IsPendingKillOrUnreachable())
 				{
-					Result.Value = TStrongObjectPtr< UDatasmithObjectTemplate >(UDatasmithObjectTemplate::GetDifference(DestinationObject, SourceTemplate.Get()));
+					Result.Value = TStrongObjectPtr< UDatasmithObjectTemplate >(UDatasmithObjectTemplate::GetDifference( DestinationObject, SourceTemplate.Get()));
 				}
 				else
 				{
@@ -500,9 +501,9 @@ namespace DatasmithImporterImpl
 	 * - Duplicate A over A'
 	 * - ApplyMigratedTemplates(A')
 	 */
-	static void ApplyMigratedTemplates(TArray< FMigratedTemplatePairType >& MigratedTemplates, UObject* DestinationObject)
+	static void ApplyMigratedTemplates( TArray< FMigratedTemplatePairType >& MigratedTemplates, UObject* DestinationObject )
 	{
-		for (FMigratedTemplatePairType& MigratedTemplate : MigratedTemplates)
+		for ( FMigratedTemplatePairType& MigratedTemplate : MigratedTemplates )
 		{
 			UDatasmithObjectTemplate* SourceTemplate = MigratedTemplate.Key.Get();
 			UDatasmithObjectTemplate* DestinationTemplate = MigratedTemplate.Value.Get();
@@ -513,32 +514,32 @@ namespace DatasmithImporterImpl
 		}
 	}
 
-	static UObject* FinalizeAsset(UObject* SourceAsset, const TCHAR* AssetPath, UObject* ExistingAsset, TMap< UObject*, UObject* >* ReferencesToRemap)
+	static UObject* FinalizeAsset( UObject* SourceAsset, const TCHAR* AssetPath, UObject* ExistingAsset, TMap< UObject*, UObject* >* ReferencesToRemap )
 	{
-		if (ReferencesToRemap)
+		if ( ReferencesToRemap )
 		{
-			FixReferencesForObject(SourceAsset, *ReferencesToRemap);
+			FixReferencesForObject( SourceAsset, *ReferencesToRemap );
 		}
 
-		TArray< FMigratedTemplatePairType > MigratedTemplates = MigrateTemplates(SourceAsset, ExistingAsset, ReferencesToRemap, false);
+		TArray< FMigratedTemplatePairType > MigratedTemplates = MigrateTemplates( SourceAsset, ExistingAsset, ReferencesToRemap, false );
 
-		UObject* FinalAsset = PublicizeAsset(SourceAsset, AssetPath, ExistingAsset);
+		UObject* FinalAsset = PublicizeAsset( SourceAsset, AssetPath, ExistingAsset );
 
-		ApplyMigratedTemplates(MigratedTemplates, FinalAsset);
+		ApplyMigratedTemplates( MigratedTemplates, FinalAsset );
 
-		if (ReferencesToRemap && SourceAsset && SourceAsset != FinalAsset)
+		if ( ReferencesToRemap && SourceAsset && SourceAsset != FinalAsset )
 		{
-			ReferencesToRemap->Add(SourceAsset, FinalAsset);
+			ReferencesToRemap->Add( SourceAsset, FinalAsset );
 		}
 
 		return FinalAsset;
 	}
 
-	class FActorWriter final : public FObjectWriter
+	class FActorWriter final: public FObjectWriter
 	{
 	public:
-		FActorWriter(UObject* Object, TArray< uint8 >& Bytes)
-			: FObjectWriter(Bytes)
+		FActorWriter( UObject* Object, TArray< uint8 >& Bytes )
+			: FObjectWriter( Bytes )
 		{
 			SetIsLoading(false);
 			SetIsSaving(true);
@@ -551,11 +552,11 @@ namespace DatasmithImporterImpl
 		{
 			bool bSkip = false;
 
-			if (InProperty->IsA< UObjectPropertyBase >())
+			if ( InProperty->IsA< UObjectPropertyBase >() )
 			{
 				bSkip = true;
 			}
-			else if (InProperty->HasAnyPropertyFlags(CPF_Transient) || !InProperty->HasAnyPropertyFlags(CPF_Edit | CPF_Interp))
+			else if ( InProperty->HasAnyPropertyFlags(CPF_Transient) || !InProperty->HasAnyPropertyFlags(CPF_Edit | CPF_Interp) )
 			{
 				bSkip = true;
 			}
@@ -567,8 +568,8 @@ namespace DatasmithImporterImpl
 	class FComponentWriter final : public FObjectWriter
 	{
 	public:
-		FComponentWriter(UObject* Object, TArray< uint8 >& Bytes)
-			: FObjectWriter(Bytes)
+		FComponentWriter( UObject* Object, TArray< uint8 >& Bytes )
+			: FObjectWriter( Bytes )
 		{
 			SetIsLoading(false);
 			SetIsSaving(true);
@@ -581,7 +582,7 @@ namespace DatasmithImporterImpl
 		{
 			bool bSkip = false;
 
-			if (InProperty->HasAnyPropertyFlags(CPF_Transient) || !InProperty->HasAnyPropertyFlags(CPF_Edit | CPF_Interp))
+			if ( InProperty->HasAnyPropertyFlags(CPF_Transient) || !InProperty->HasAnyPropertyFlags(CPF_Edit | CPF_Interp) )
 			{
 				bSkip = true;
 			}
@@ -593,20 +594,20 @@ namespace DatasmithImporterImpl
 	static void DeleteImportSceneActorIfNeeded(FDatasmithActorImportContext& ActorContext, bool bForce = false)
 	{
 		ADatasmithSceneActor*& ImportSceneActor = ActorContext.ImportSceneActor;
-		if (!ActorContext.FinalSceneActors.Contains(ImportSceneActor) || bForce)
+		if ( !ActorContext.FinalSceneActors.Contains(ImportSceneActor) || bForce )
 		{
-			if (ImportSceneActor)
+			if ( ImportSceneActor )
 			{
-				for (const auto& Pair : ImportSceneActor->RelatedActors)
+				for ( const auto& Pair : ImportSceneActor->RelatedActors )
 				{
 					AActor* RelatedActor = Pair.Value.Get();
-					if (RelatedActor)
+					if ( RelatedActor )
 					{
-						FDatasmithImporterUtils::DeleteActor(*RelatedActor);
+						FDatasmithImporterUtils::DeleteActor( *RelatedActor );
 					}
 				}
 
-				FDatasmithImporterUtils::DeleteActor(*ImportSceneActor);
+				FDatasmithImporterUtils::DeleteActor( *ImportSceneActor );
 
 				// Null also the ImportSceneActor from the Actor Context because it's a ref to it.
 				ImportSceneActor = nullptr;
@@ -614,13 +615,13 @@ namespace DatasmithImporterImpl
 		}
 	}
 
-	static UActorComponent* PublicizeComponent(UActorComponent& SourceComponent, UActorComponent* DestinationComponent, AActor& DestinationActor, TMap< UObject*, UObject* >& ReferencesToRemap, USceneComponent* DestinationParent = nullptr)
+	static UActorComponent* PublicizeComponent(UActorComponent& SourceComponent, UActorComponent* DestinationComponent, AActor& DestinationActor, TMap< UObject*, UObject* >& ReferencesToRemap, USceneComponent* DestinationParent = nullptr )
 	{
 		TRACE_CPUPROFILER_EVENT_SCOPE(DatasmithImporterImpl::PublicizeComponent);
 
 		SafetyGarbageCollect();
 
-		if (!SourceComponent.HasAnyFlags(RF_Transient | RF_TextExportTransient))
+		if ( !SourceComponent.HasAnyFlags( RF_Transient | RF_TextExportTransient ) )
 		{
 			if (!DestinationComponent || DestinationComponent->IsPendingKillOrUnreachable())
 			{
@@ -630,9 +631,9 @@ namespace DatasmithImporterImpl
 					DestinationComponent->Rename(nullptr, nullptr, REN_DontCreateRedirectors | REN_NonTransactional);
 				}
 
-				if (UActorComponent* OldComponent = static_cast<UActorComponent*>(FindObjectWithOuter(&DestinationActor, UActorComponent::StaticClass(), SourceComponent.GetFName())))
+				if ( UActorComponent* OldComponent = static_cast<UActorComponent*>( FindObjectWithOuter( &DestinationActor, UActorComponent::StaticClass(), SourceComponent.GetFName() ) ) )
 				{
-					OldComponent->DestroyComponent(true);
+					OldComponent->DestroyComponent( true );
 					// Change the name of the old component so that the new object won't recycle the old one.
 					OldComponent->Rename(nullptr, nullptr, REN_DontCreateRedirectors | REN_NonTransactional);
 				}
@@ -643,13 +644,13 @@ namespace DatasmithImporterImpl
 
 				if (USceneComponent* NewSceneComponent = Cast<USceneComponent>(DestinationComponent))
 				{
-					if (!DestinationActor.GetRootComponent())
+					if ( !DestinationActor.GetRootComponent() )
 					{
 						DestinationActor.SetRootComponent(NewSceneComponent);
 					}
-					if (DestinationParent)
+					if ( DestinationParent )
 					{
-						NewSceneComponent->AttachToComponent(DestinationParent, FAttachmentTransformRules::KeepRelativeTransform);
+						NewSceneComponent->AttachToComponent( DestinationParent, FAttachmentTransformRules::KeepRelativeTransform );
 					}
 				}
 
@@ -670,14 +671,14 @@ namespace DatasmithImporterImpl
 			// #ueent_todo: we shouldn't be copying instanced object pointers in the first place
 			UDatasmithAssetUserData* SourceAssetUserData = DestinationComponent->GetAssetUserData< UDatasmithAssetUserData >();
 
-			if (SourceAssetUserData)
+			if ( SourceAssetUserData )
 			{
-				UDatasmithAssetUserData* DestinationAssetUserData = DuplicateObject< UDatasmithAssetUserData >(SourceAssetUserData, DestinationComponent);
-				DestinationComponent->RemoveUserDataOfClass(UDatasmithAssetUserData::StaticClass());
-				DestinationComponent->AddAssetUserData(DestinationAssetUserData);
+				UDatasmithAssetUserData* DestinationAssetUserData = DuplicateObject< UDatasmithAssetUserData >( SourceAssetUserData, DestinationComponent );
+				DestinationComponent->RemoveUserDataOfClass( UDatasmithAssetUserData::StaticClass() );
+				DestinationComponent->AddAssetUserData( DestinationAssetUserData );
 			}
 
-			ReferencesToRemap.Add(&SourceComponent, DestinationComponent);
+			ReferencesToRemap.Add( &SourceComponent, DestinationComponent );
 
 			return DestinationComponent;
 		}
@@ -685,18 +686,18 @@ namespace DatasmithImporterImpl
 		return nullptr;
 	}
 
-	static void FinalizeSceneComponent(FDatasmithImportContext& ImportContext, USceneComponent& SourceComponent, AActor& DestinationActor, USceneComponent* DestinationParent, TMap<UObject *, UObject *>& ReferencesToRemap)
+	static void FinalizeSceneComponent(FDatasmithImportContext& ImportContext, USceneComponent& SourceComponent, AActor& DestinationActor, USceneComponent* DestinationParent, TMap<UObject *, UObject *>& ReferencesToRemap )
 	{
-		USceneComponent* DestinationComponent = static_cast<USceneComponent*>(FindObjectWithOuter(&DestinationActor, SourceComponent.GetClass(), SourceComponent.GetFName()));
+		USceneComponent* DestinationComponent = static_cast<USceneComponent*>( FindObjectWithOuter( &DestinationActor, SourceComponent.GetClass(), SourceComponent.GetFName() ) );
 		FName SourceComponentDatasmithId = FDatasmithImporterUtils::GetDatasmithElementId(&SourceComponent);
 
-		if (SourceComponentDatasmithId.IsNone())
+		if ( SourceComponentDatasmithId.IsNone() )
 		{
 			// This component is not tracked by datasmith
-			if (!DestinationComponent || DestinationComponent->IsPendingKillOrUnreachable())
+			if ( !DestinationComponent || DestinationComponent->IsPendingKillOrUnreachable() )
 			{
-				DestinationComponent = static_cast<USceneComponent*> (PublicizeComponent(SourceComponent, DestinationComponent, DestinationActor, ReferencesToRemap, DestinationParent));
-				if (DestinationComponent)
+				DestinationComponent = static_cast<USceneComponent*> ( PublicizeComponent(SourceComponent, DestinationComponent, DestinationActor, ReferencesToRemap, DestinationParent) );
+				if ( DestinationComponent )
 				{
 					// Put back the components in a proper state
 					DestinationComponent->UpdateComponentToWorld();
@@ -714,24 +715,24 @@ namespace DatasmithImporterImpl
 				false
 			);
 
-			DestinationComponent = static_cast<USceneComponent*> (PublicizeComponent(SourceComponent, DestinationComponent, DestinationActor, ReferencesToRemap, DestinationParent));
+			DestinationComponent = static_cast< USceneComponent* > ( PublicizeComponent( SourceComponent, DestinationComponent, DestinationActor, ReferencesToRemap, DestinationParent ) );
 
-			if (DestinationComponent)
+			if ( DestinationComponent )
 			{
 				// Put back the components in a proper state (Unfortunately without this the set relative transform might not work)
 				DestinationComponent->UpdateComponentToWorld();
-				ApplyMigratedTemplates(MigratedTemplates, DestinationComponent);
+				ApplyMigratedTemplates( MigratedTemplates, DestinationComponent );
 				DestinationComponent->PostEditChange();
 			}
 		}
 
-		USceneComponent* AttachParentForChildren = DestinationComponent ? DestinationComponent : DestinationParent;
-		for (USceneComponent* Child : SourceComponent.GetAttachChildren())
+		USceneComponent* AttachParentForChildren = DestinationComponent  ? DestinationComponent : DestinationParent;
+		for ( USceneComponent* Child : SourceComponent.GetAttachChildren() )
 		{
 			// Only finalize components that are from the same actor
-			if (Child && Child->GetOuter() == SourceComponent.GetOuter())
+			if ( Child && Child->GetOuter() == SourceComponent.GetOuter() )
 			{
-				FinalizeSceneComponent(ImportContext, *Child, DestinationActor, AttachParentForChildren, ReferencesToRemap);
+				FinalizeSceneComponent( ImportContext, *Child, DestinationActor, AttachParentForChildren, ReferencesToRemap );
 			}
 		}
 	}
@@ -741,31 +742,31 @@ namespace DatasmithImporterImpl
 		USceneComponent* ParentComponent = nullptr;
 
 		// Find the parent component
-		UObject** ObjectPtr = ReferencesToRemap.Find(SourceActor.GetRootComponent()->GetAttachParent());
-		if (ObjectPtr)
+		UObject** ObjectPtr = ReferencesToRemap.Find( SourceActor.GetRootComponent()->GetAttachParent() );
+		if ( ObjectPtr )
 		{
-			ParentComponent = Cast<USceneComponent>(*ObjectPtr);
+			ParentComponent = Cast<USceneComponent>( *ObjectPtr );
 		}
 
 		// Finalize the scene components recursively
 		{
 			USceneComponent* RootComponent = SourceActor.GetRootComponent();
-			if (RootComponent)
+			if ( RootComponent )
 			{
-				FinalizeSceneComponent(ImportContext, *RootComponent, DestinationActor, ParentComponent, ReferencesToRemap);
+				FinalizeSceneComponent( ImportContext, *RootComponent, DestinationActor, ParentComponent, ReferencesToRemap );
 			}
 		}
 
 
-		for (UActorComponent* SourceComponent : SourceActor.GetComponents())
+		for ( UActorComponent* SourceComponent : SourceActor.GetComponents() )
 		{
 			// Only the non scene component haven't been finalize
-			if (SourceComponent && !SourceComponent->GetClass()->IsChildOf<USceneComponent>())
+			if ( SourceComponent && !SourceComponent->GetClass()->IsChildOf<USceneComponent>() )
 			{
-				UActorComponent* DestinationComponent = static_cast<UActorComponent*>(FindObjectWithOuter(&DestinationActor, SourceComponent->GetClass(), SourceComponent->GetFName()));
+				UActorComponent* DestinationComponent = static_cast< UActorComponent* >( FindObjectWithOuter( &DestinationActor, SourceComponent->GetClass(), SourceComponent->GetFName() ) );
 				if (!DestinationComponent)
 				{
-					DestinationComponent = PublicizeComponent(*SourceComponent, DestinationComponent, DestinationActor, ReferencesToRemap);
+					DestinationComponent = PublicizeComponent( *SourceComponent, DestinationComponent, DestinationActor, ReferencesToRemap );
 				}
 			}
 		}
@@ -782,7 +783,7 @@ namespace DatasmithImporterImpl
 			UMaterialInterface* CurrentMaterialInterface = ImportedMaterialPair.Value;
 			UMaterial* BaseMaterial = CurrentMaterialInterface->GetMaterial();
 
-			if (!TextureParametersToConvertMap.Contains(BaseMaterial))
+			if(!TextureParametersToConvertMap.Contains(BaseMaterial))
 			{
 				bool bRequiresTextureCheck = false;
 				TArray<FMaterialParameterInfo> OutParameterInfo;
@@ -871,18 +872,18 @@ namespace DatasmithImporterImpl
 	}
 }
 
-void FDatasmithImporter::ImportStaticMeshes(FDatasmithImportContext& ImportContext)
+void FDatasmithImporter::ImportStaticMeshes( FDatasmithImportContext& ImportContext )
 {
 	const int32 StaticMeshesCount = ImportContext.FilteredScene->GetMeshesCount();
 
-	if (!ImportContext.Options->BaseOptions.bIncludeGeometry || StaticMeshesCount == 0)
+	if ( !ImportContext.Options->BaseOptions.bIncludeGeometry || StaticMeshesCount == 0 )
 	{
 		return;
 	}
 
 	TRACE_CPUPROFILER_EVENT_SCOPE(FDatasmithImporter::ImportStaticMeshes);
 
-	FScopedSlowTask Progress(StaticMeshesCount, LOCTEXT("ImportStaticMeshes", "Importing Static Meshes..."), true, *ImportContext.Warn);
+	FScopedSlowTask Progress(StaticMeshesCount, LOCTEXT("ImportStaticMeshes", "Importing Static Meshes..."), true, *ImportContext.Warn );
 	Progress.MakeDialog(true);
 
 	TMap<TSharedRef<IDatasmithMeshElement>, TFuture<FDatasmithMeshElementPayload*>> MeshElementPayloads;
@@ -905,9 +906,9 @@ void FDatasmithImporter::ImportStaticMeshes(FDatasmithImportContext& ImportConte
 				continue;
 			}
 
-			TSharedRef<IDatasmithMeshElement> MeshElement = ImportContext.FilteredScene->GetMesh(MeshIndex).ToSharedRef();
+			TSharedRef<IDatasmithMeshElement> MeshElement = ImportContext.FilteredScene->GetMesh( MeshIndex ).ToSharedRef();
 
-			UStaticMesh*& ImportedStaticMesh = ImportContext.ImportedStaticMeshes.FindOrAdd(MeshElement);
+			UStaticMesh*& ImportedStaticMesh = ImportContext.ImportedStaticMeshes.FindOrAdd( MeshElement );
 
 			// We still have factories that are importing the UStaticMesh on their own, so check if it's already imported here
 			if (ImportedStaticMesh == nullptr)
@@ -918,15 +919,15 @@ void FDatasmithImporter::ImportStaticMeshes(FDatasmithImportContext& ImportConte
 					Async(
 						EAsyncExecution::LargeThreadPool,
 						[&ImportContext, MeshElement]() -> FDatasmithMeshElementPayload*
-				{
-					if (ImportContext.bUserCancelled)
-					{
-						return nullptr;
-					}
+						{
+							if (ImportContext.bUserCancelled)
+							{
+								return nullptr;
+							}
 
-					TUniquePtr<FDatasmithMeshElementPayload> MeshPayload = MakeUnique<FDatasmithMeshElementPayload>();
-					return ImportContext.SceneTranslator->LoadStaticMesh(MeshElement, *MeshPayload) ? MeshPayload.Release() : nullptr;
-				}
+							TUniquePtr<FDatasmithMeshElementPayload> MeshPayload = MakeUnique<FDatasmithMeshElementPayload>();
+							return ImportContext.SceneTranslator->LoadStaticMesh(MeshElement, *MeshPayload) ? MeshPayload.Release() : nullptr;
+						}
 					)
 				);
 			}
@@ -934,19 +935,19 @@ void FDatasmithImporter::ImportStaticMeshes(FDatasmithImportContext& ImportConte
 	}
 
 	// This pass will wait on the futures we got from the first pass async tasks
-	for (int32 MeshIndex = 0; MeshIndex < StaticMeshesCount && !ImportContext.bUserCancelled; ++MeshIndex)
+	for ( int32 MeshIndex = 0; MeshIndex < StaticMeshesCount && !ImportContext.bUserCancelled; ++MeshIndex )
 	{
 		ImportContext.bUserCancelled |= ImportContext.Warn->ReceivedUserCancel();
 
-		TSharedRef< IDatasmithMeshElement > MeshElement = ImportContext.FilteredScene->GetMesh(MeshIndex).ToSharedRef();
+		TSharedRef< IDatasmithMeshElement > MeshElement = ImportContext.FilteredScene->GetMesh( MeshIndex ).ToSharedRef();
 
-		Progress.EnterProgressFrame(1.f, FText::FromString(FString::Printf(TEXT("Importing static mesh %d/%d (%s) ..."), MeshIndex + 1, StaticMeshesCount, MeshElement->GetLabel())));
+		Progress.EnterProgressFrame( 1.f, FText::FromString( FString::Printf( TEXT("Importing static mesh %d/%d (%s) ..."), MeshIndex + 1, StaticMeshesCount, MeshElement->GetLabel() ) ) );
 
 		UStaticMesh* ExistingStaticMesh = nullptr;
 
 		if (ImportContext.SceneAsset)
 		{
-			TSoftObjectPtr< UStaticMesh >* ExistingStaticMeshPtr = ImportContext.SceneAsset->StaticMeshes.Find(MeshElement->GetName());
+			TSoftObjectPtr< UStaticMesh >* ExistingStaticMeshPtr = ImportContext.SceneAsset->StaticMeshes.Find( MeshElement->GetName() );
 
 			if (ExistingStaticMeshPtr)
 			{
@@ -974,34 +975,34 @@ void FDatasmithImporter::ImportStaticMeshes(FDatasmithImportContext& ImportConte
 	}
 
 	//Just make sure there is no async task left running in case of a cancellation
-	for (const TPair<TSharedRef<IDatasmithMeshElement>, TFuture<FDatasmithMeshElementPayload*>> & Kvp : MeshElementPayloads)
+	for ( const TPair<TSharedRef<IDatasmithMeshElement>, TFuture<FDatasmithMeshElementPayload*>> & Kvp : MeshElementPayloads)
 	{
 		// Wait for the result and delete it when getting out of scope
 		TUniquePtr<FDatasmithMeshElementPayload> MeshPayloadPtr(Kvp.Value.Get());
 	}
 
-	TMap< TSharedRef< IDatasmithMeshElement >, float > LightmapWeights = FDatasmithStaticMeshImporter::CalculateMeshesLightmapWeights(ImportContext.Scene.ToSharedRef());
+	TMap< TSharedRef< IDatasmithMeshElement >, float > LightmapWeights = FDatasmithStaticMeshImporter::CalculateMeshesLightmapWeights( ImportContext.Scene.ToSharedRef() );
 
-	for (TPair< TSharedRef< IDatasmithMeshElement >, UStaticMesh* >& ImportedStaticMeshPair : ImportContext.ImportedStaticMeshes)
+	for ( TPair< TSharedRef< IDatasmithMeshElement >, UStaticMesh* >& ImportedStaticMeshPair : ImportContext.ImportedStaticMeshes )
 	{
-		FDatasmithStaticMeshImporter::SetupStaticMesh(ImportContext.AssetsContext, ImportedStaticMeshPair.Key, ImportedStaticMeshPair.Value, ImportContext.Options->BaseOptions.StaticMeshOptions, LightmapWeights[ImportedStaticMeshPair.Key]);
+		FDatasmithStaticMeshImporter::SetupStaticMesh( ImportContext.AssetsContext, ImportedStaticMeshPair.Key, ImportedStaticMeshPair.Value, ImportContext.Options->BaseOptions.StaticMeshOptions, LightmapWeights[ ImportedStaticMeshPair.Key ] );
 	}
 }
 
-UStaticMesh* FDatasmithImporter::ImportStaticMesh(FDatasmithImportContext& ImportContext, TSharedRef< IDatasmithMeshElement > MeshElement, UStaticMesh* ExistingStaticMesh, FDatasmithMeshElementPayload* MeshPayload)
+UStaticMesh* FDatasmithImporter::ImportStaticMesh( FDatasmithImportContext& ImportContext, TSharedRef< IDatasmithMeshElement > MeshElement, UStaticMesh* ExistingStaticMesh, FDatasmithMeshElementPayload* MeshPayload)
 {
-	if (!ImportContext.AssetsContext.StaticMeshesFinalPackage || ImportContext.AssetsContext.StaticMeshesFinalPackage->GetFName() == NAME_None || ImportContext.SceneTranslator == nullptr)
+	if ( !ImportContext.AssetsContext.StaticMeshesFinalPackage || ImportContext.AssetsContext.StaticMeshesFinalPackage->GetFName() == NAME_None || ImportContext.SceneTranslator == nullptr)
 	{
 		return nullptr;
 	}
 
 	TRACE_CPUPROFILER_EVENT_SCOPE(FDatasmithImporter::ImportStaticMesh);
 
-	UStaticMesh*& ImportedStaticMesh = ImportContext.ImportedStaticMeshes.FindOrAdd(MeshElement);
+	UStaticMesh*& ImportedStaticMesh = ImportContext.ImportedStaticMeshes.FindOrAdd( MeshElement );
 
 	TArray<UDatasmithAdditionalData*> AdditionalData;
 
-	if (ImportedStaticMesh == nullptr) // We still have factories that are importing the UStaticMesh on their own, so check if it's already imported here
+	if ( ImportedStaticMesh == nullptr ) // We still have factories that are importing the UStaticMesh on their own, so check if it's already imported here
 	{
 		FDatasmithMeshElementPayload LocalMeshPayload;
 		if (MeshPayload == nullptr)
@@ -1010,7 +1011,7 @@ UStaticMesh* FDatasmithImporter::ImportStaticMesh(FDatasmithImportContext& Impor
 			MeshPayload = &LocalMeshPayload;
 		}
 
-		ImportedStaticMesh = FDatasmithStaticMeshImporter::ImportStaticMesh(MeshElement, *MeshPayload, ImportContext.ObjectFlags & ~RF_Public, ImportContext.Options->BaseOptions.StaticMeshOptions, ImportContext.AssetsContext, ExistingStaticMesh);
+		ImportedStaticMesh = FDatasmithStaticMeshImporter::ImportStaticMesh( MeshElement, *MeshPayload, ImportContext.ObjectFlags & ~RF_Public, ImportContext.Options->BaseOptions.StaticMeshOptions, ImportContext.AssetsContext, ExistingStaticMesh );
 		AdditionalData = MoveTemp(MeshPayload->AdditionalData);
 
 		// Creation of static mesh failed, remove it from the list of importer mesh elements
@@ -1021,11 +1022,11 @@ UStaticMesh* FDatasmithImporter::ImportStaticMesh(FDatasmithImportContext& Impor
 		}
 	}
 
-	CreateStaticMeshAssetImportData(ImportContext, MeshElement, ImportedStaticMesh, AdditionalData);
+	CreateStaticMeshAssetImportData( ImportContext, MeshElement, ImportedStaticMesh, AdditionalData );
 
-	ImportMetaDataForObject(ImportContext, MeshElement, ImportedStaticMesh);
+	ImportMetaDataForObject( ImportContext, MeshElement, ImportedStaticMesh );
 
-	if (MeshElement->GetLightmapSourceUV() >= MAX_MESH_TEXTURE_COORDS_MD)
+	if ( MeshElement->GetLightmapSourceUV() >= MAX_MESH_TEXTURE_COORDS_MD )
 	{
 		FFormatNamedArguments FormatArgs;
 		FormatArgs.Add(TEXT("SourceUV"), FText::FromString(FString::FromInt(MeshElement->GetLightmapSourceUV())));
@@ -1037,11 +1038,11 @@ UStaticMesh* FDatasmithImporter::ImportStaticMesh(FDatasmithImportContext& Impor
 }
 
 
-UStaticMesh* FDatasmithImporter::FinalizeStaticMesh(UStaticMesh* SourceStaticMesh, const TCHAR* StaticMeshesFolderPath, UStaticMesh* ExistingStaticMesh, TMap< UObject*, UObject* >* ReferencesToRemap, bool bBuild)
+UStaticMesh* FDatasmithImporter::FinalizeStaticMesh( UStaticMesh* SourceStaticMesh, const TCHAR* StaticMeshesFolderPath, UStaticMesh* ExistingStaticMesh, TMap< UObject*, UObject* >* ReferencesToRemap, bool bBuild)
 {
 	using namespace DatasmithImporterImpl;
 
-	UStaticMesh* DestinationStaticMesh = Cast< UStaticMesh >(FinalizeAsset(SourceStaticMesh, StaticMeshesFolderPath, ExistingStaticMesh, ReferencesToRemap));
+	UStaticMesh* DestinationStaticMesh = Cast< UStaticMesh >( FinalizeAsset( SourceStaticMesh, StaticMeshesFolderPath, ExistingStaticMesh, ReferencesToRemap ) );
 
 	if (bBuild)
 	{
@@ -1053,19 +1054,19 @@ UStaticMesh* FDatasmithImporter::FinalizeStaticMesh(UStaticMesh* SourceStaticMes
 
 void FDatasmithImporter::CreateStaticMeshAssetImportData(FDatasmithImportContext& InContext, TSharedRef< IDatasmithMeshElement > MeshElement, UStaticMesh* ImportedStaticMesh, TArray<UDatasmithAdditionalData*>& AdditionalData)
 {
-	UDatasmithStaticMeshImportData::DefaultOptionsPair ImportOptions = UDatasmithStaticMeshImportData::DefaultOptionsPair(InContext.Options->BaseOptions.StaticMeshOptions, InContext.Options->BaseOptions.AssetOptions);
+	UDatasmithStaticMeshImportData::DefaultOptionsPair ImportOptions = UDatasmithStaticMeshImportData::DefaultOptionsPair( InContext.Options->BaseOptions.StaticMeshOptions, InContext.Options->BaseOptions.AssetOptions );
 
-	UDatasmithStaticMeshImportData* MeshImportData = UDatasmithStaticMeshImportData::GetImportDataForStaticMesh(ImportedStaticMesh, ImportOptions);
+	UDatasmithStaticMeshImportData* MeshImportData = UDatasmithStaticMeshImportData::GetImportDataForStaticMesh( ImportedStaticMesh, ImportOptions );
 
-	if (MeshImportData)
+	if ( MeshImportData )
 	{
 		// Update the import data source file and set the mesh hash
 		// #ueent_todo FH: piggybacking off of the SourceData file hash for now, until we have custom derived AssetImportData properly serialize to the AssetRegistry
-		FMD5Hash Hash = MeshElement->CalculateElementHash(false);
-		MeshImportData->Update(InContext.Options->FilePath, &Hash);
+		FMD5Hash Hash = MeshElement->CalculateElementHash( false );
+		MeshImportData->Update( InContext.Options->FilePath, &Hash );
 
 		// Set the final outer // #ueent_review: propagate flags of outer?
-		for (UDatasmithAdditionalData* Data : AdditionalData)
+		for (UDatasmithAdditionalData* Data: AdditionalData)
 		{
 			Data->Rename(nullptr, MeshImportData);
 		}
@@ -1073,11 +1074,11 @@ void FDatasmithImporter::CreateStaticMeshAssetImportData(FDatasmithImportContext
 	}
 }
 
-void FDatasmithImporter::ImportTextures(FDatasmithImportContext& ImportContext)
+void FDatasmithImporter::ImportTextures( FDatasmithImportContext& ImportContext )
 {
 	TRACE_CPUPROFILER_EVENT_SCOPE(FDatasmithImporter::ImportTextures);
 
-	DatasmithImporterImpl::SetTexturesMode(ImportContext);
+	DatasmithImporterImpl::SetTexturesMode( ImportContext );
 
 	const int32 TexturesCount = ImportContext.FilteredScene->GetTexturesCount();
 
@@ -1088,94 +1089,143 @@ void FDatasmithImporter::ImportTextures(FDatasmithImportContext& ImportContext)
 	{
 		FDatasmithTextureImporter DatasmithTextureImporter(ImportContext);
 
-		for (int32 i = 0; i < TexturesCount && !ImportContext.bUserCancelled; i++)
+		TArray<TSharedPtr< IDatasmithTextureElement >> FilteredTextureElements;
+		for ( int32 i = 0; i < TexturesCount; i++ )
 		{
-			ImportContext.bUserCancelled |= ImportContext.Warn->ReceivedUserCancel();
-
 			TSharedPtr< IDatasmithTextureElement > TextureElement = ImportContext.FilteredScene->GetTexture(i);
 
-			if (!TextureElement)
+			if ( !TextureElement )
 			{
 				continue;
 			}
 
-			Progress.EnterProgressFrame(1.f, FText::FromString(FString::Printf(TEXT("Importing texture %d/%d (%s) ..."), i + 1, TexturesCount, TextureElement->GetLabel())));
+			FilteredTextureElements.Add(TextureElement);
+		}
+
+		FDatasmithTextureResize::Initialize();
 
-			UTexture* ExistingTexture = nullptr;
+		struct FAsyncData
+		{
+			FString       Extension;
+			TArray<uint8> TextureData;
+			TFuture<bool> Result;
+		};
+		TArray<FAsyncData> AsyncData;
+		AsyncData.SetNum(FilteredTextureElements.Num());
 
-			if (ImportContext.SceneAsset)
+		for ( int32 TextureIndex = 0; TextureIndex < FilteredTextureElements.Num(); TextureIndex++ )
+		{
+			ImportContext.bUserCancelled |= ImportContext.Warn->ReceivedUserCancel();
+
+			AsyncData[TextureIndex].Result = 
+				Async(
+					EAsyncExecution::LargeThreadPool,
+					[&ImportContext, &AsyncData, &FilteredTextureElements, &DatasmithTextureImporter, TextureIndex]()
+					{
+						if (ImportContext.bUserCancelled)
+						{
+							return false;
+						}
+
+						return DatasmithTextureImporter.GetTextureData(FilteredTextureElements[TextureIndex], AsyncData[TextureIndex].TextureData, AsyncData[TextureIndex].Extension);
+					}
+				);
+		}
+
+		for ( int32 TextureIndex = 0; TextureIndex < FilteredTextureElements.Num(); TextureIndex++ )
+		{
+			if (ImportContext.bUserCancelled |= ImportContext.Warn->ReceivedUserCancel())
 			{
-				TSoftObjectPtr< UTexture >* ExistingTexturePtr = ImportContext.SceneAsset->Textures.Find(TextureElement->GetName());
+				// If operation has been canceled, just wait for other threads to also cancel
+				AsyncData[TextureIndex].Result.Wait();
+			}
+			else
+			{
+				const TSharedPtr< IDatasmithTextureElement >& TextureElement = FilteredTextureElements[TextureIndex];
+
+				Progress.EnterProgressFrame( 1.f, FText::FromString( FString::Printf( TEXT("Importing texture %d/%d (%s) ..."), TextureIndex + 1, FilteredTextureElements.Num(), TextureElement->GetLabel() ) ) );
+
+				UTexture* ExistingTexture = nullptr;
+
+				if ( ImportContext.SceneAsset )
+				{
+					TSoftObjectPtr< UTexture >* ExistingTexturePtr = ImportContext.SceneAsset->Textures.Find( TextureElement->GetName() );
+
+					if ( ExistingTexturePtr )
+					{
+						ExistingTexture = ExistingTexturePtr->LoadSynchronous();
+					}
+				}
 
-				if (ExistingTexturePtr)
+				if (AsyncData[TextureIndex].Result.Get())
 				{
-					ExistingTexture = ExistingTexturePtr->LoadSynchronous();
+					ImportTexture( ImportContext, DatasmithTextureImporter, TextureElement.ToSharedRef(), ExistingTexture, AsyncData[TextureIndex].TextureData, AsyncData[TextureIndex].Extension );
 				}
 			}
 
-			ImportTexture(ImportContext, DatasmithTextureImporter, TextureElement.ToSharedRef(), ExistingTexture);
-
+			// Release memory as soon as possible
+			AsyncData[TextureIndex].TextureData.Empty();
 		}
 	}
 }
 
-UTexture* FDatasmithImporter::ImportTexture(FDatasmithImportContext& ImportContext, FDatasmithTextureImporter& DatasmithTextureImporter, TSharedRef< IDatasmithTextureElement > TextureElement, UTexture* ExistingTexture)
+UTexture* FDatasmithImporter::ImportTexture( FDatasmithImportContext& ImportContext, FDatasmithTextureImporter& DatasmithTextureImporter, TSharedRef< IDatasmithTextureElement > TextureElement, UTexture* ExistingTexture, const TArray<uint8>& TextureData, const FString& Extension )
 {
 	TRACE_CPUPROFILER_EVENT_SCOPE(FDatasmithImporter::ImportTexture);
 
-	UTexture*& ImportedTexture = ImportContext.ImportedTextures.FindOrAdd(TextureElement);
-	ImportedTexture = DatasmithTextureImporter.CreateTexture(TextureElement);
+	UTexture*& ImportedTexture = ImportContext.ImportedTextures.FindOrAdd( TextureElement );
+	ImportedTexture = DatasmithTextureImporter.CreateTexture( TextureElement, TextureData, Extension );
 
 	if (ImportedTexture == nullptr)
 	{
-		ImportContext.ImportedTextures.Remove(TextureElement);
+		ImportContext.ImportedTextures.Remove( TextureElement );
 		return nullptr;
 	}
 
-	ImportMetaDataForObject(ImportContext, TextureElement, ImportedTexture);
+	ImportMetaDataForObject( ImportContext, TextureElement, ImportedTexture );
 
 	return ImportedTexture;
 }
 
-UTexture* FDatasmithImporter::FinalizeTexture(UTexture* SourceTexture, const TCHAR* TexturesFolderPath, UTexture* ExistingTexture, TMap< UObject*, UObject* >* ReferencesToRemap)
+UTexture* FDatasmithImporter::FinalizeTexture( UTexture* SourceTexture, const TCHAR* TexturesFolderPath, UTexture* ExistingTexture, TMap< UObject*, UObject* >* ReferencesToRemap )
 {
-	return Cast< UTexture >(DatasmithImporterImpl::FinalizeAsset(SourceTexture, TexturesFolderPath, ExistingTexture, ReferencesToRemap));
+	return Cast< UTexture >( DatasmithImporterImpl::FinalizeAsset( SourceTexture, TexturesFolderPath, ExistingTexture, ReferencesToRemap ) );
 }
 
-void FDatasmithImporter::ImportMaterials(FDatasmithImportContext& ImportContext)
+void FDatasmithImporter::ImportMaterials( FDatasmithImportContext& ImportContext )
 {
 	TRACE_CPUPROFILER_EVENT_SCOPE(FDatasmithImporter::ImportMaterials);
 
-	if (ImportContext.Options->MaterialConflictPolicy != EDatasmithImportAssetConflictPolicy::Ignore && ImportContext.FilteredScene->GetMaterialsCount() > 0)
+	if ( ImportContext.Options->MaterialConflictPolicy != EDatasmithImportAssetConflictPolicy::Ignore && ImportContext.FilteredScene->GetMaterialsCount() > 0 )
 	{
-		IDatasmithShaderElement::bUseRealisticFresnel = (ImportContext.Options->MaterialQuality == EDatasmithImportMaterialQuality::UseRealFresnelCurves);
-		IDatasmithShaderElement::bDisableReflectionFresnel = (ImportContext.Options->MaterialQuality == EDatasmithImportMaterialQuality::UseNoFresnelCurves);
+		IDatasmithShaderElement::bUseRealisticFresnel = ( ImportContext.Options->MaterialQuality == EDatasmithImportMaterialQuality::UseRealFresnelCurves );
+		IDatasmithShaderElement::bDisableReflectionFresnel = ( ImportContext.Options->MaterialQuality == EDatasmithImportMaterialQuality::UseNoFresnelCurves );
 
 		//Import referenced materials as MaterialFunctions first
-		for (TSharedPtr< IDatasmithBaseMaterialElement > MaterialElement : FDatasmithImporterUtils::GetOrderedListOfMaterialsReferencedByMaterials(ImportContext.FilteredScene))
+		for ( TSharedPtr< IDatasmithBaseMaterialElement > MaterialElement : FDatasmithImporterUtils::GetOrderedListOfMaterialsReferencedByMaterials( ImportContext.FilteredScene ) )
 		{
-			ImportMaterialFunction(ImportContext, MaterialElement.ToSharedRef());
+			ImportMaterialFunction(ImportContext, MaterialElement.ToSharedRef() );
 		}
 
-		ImportContext.AssetsContext.MaterialsRequirements.Empty(ImportContext.FilteredScene->GetMaterialsCount());
+		ImportContext.AssetsContext.MaterialsRequirements.Empty( ImportContext.FilteredScene->GetMaterialsCount() );
 
-		for (int32 MaterialIndex = 0; MaterialIndex < ImportContext.FilteredScene->GetMaterialsCount(); ++MaterialIndex)
+		for ( int32 MaterialIndex = 0; MaterialIndex < ImportContext.FilteredScene->GetMaterialsCount(); ++MaterialIndex )
 		{
-			TSharedRef< IDatasmithBaseMaterialElement > MaterialElement = ImportContext.FilteredScene->GetMaterial(MaterialIndex).ToSharedRef();
+			TSharedRef< IDatasmithBaseMaterialElement > MaterialElement = ImportContext.FilteredScene->GetMaterial( MaterialIndex ).ToSharedRef();
 
 			UMaterialInterface* ExistingMaterial = nullptr;
 
-			if (ImportContext.SceneAsset)
+			if ( ImportContext.SceneAsset )
 			{
-				TSoftObjectPtr< UMaterialInterface >* ExistingMaterialPtr = ImportContext.SceneAsset->Materials.Find(MaterialElement->GetName());
+				TSoftObjectPtr< UMaterialInterface >* ExistingMaterialPtr = ImportContext.SceneAsset->Materials.Find( MaterialElement->GetName() );
 
-				if (ExistingMaterialPtr)
+				if ( ExistingMaterialPtr )
 				{
 					ExistingMaterial = ExistingMaterialPtr->LoadSynchronous();
 				}
 			}
 
-			ImportMaterial(ImportContext, MaterialElement, ExistingMaterial);
+			ImportMaterial( ImportContext, MaterialElement, ExistingMaterial );
 		}
 
 		// IMPORTANT: FGlobalComponentReregisterContext destructor will de-register and re-register all UActorComponent present in the world
@@ -1188,14 +1238,14 @@ void FDatasmithImporter::ImportMaterials(FDatasmithImportContext& ImportContext)
 
 UMaterialFunction* FDatasmithImporter::ImportMaterialFunction(FDatasmithImportContext& ImportContext, TSharedRef< IDatasmithBaseMaterialElement > MaterialElement)
 {
-	UMaterialFunction* ImportedMaterialFunction = FDatasmithMaterialImporter::CreateMaterialFunction(ImportContext, MaterialElement);
+	UMaterialFunction* ImportedMaterialFunction = FDatasmithMaterialImporter::CreateMaterialFunction( ImportContext, MaterialElement );
 
-	if (!ImportedMaterialFunction)
+	if (!ImportedMaterialFunction )
 	{
 		return nullptr;
 	}
 
-	ImportContext.ImportedMaterialFunctions.Add(MaterialElement) = ImportedMaterialFunction;
+	ImportContext.ImportedMaterialFunctions.Add( MaterialElement ) = ImportedMaterialFunction;
 
 	return ImportedMaterialFunction;
 }
@@ -1206,14 +1256,14 @@ UMaterialFunction* FDatasmithImporter::FinalizeMaterialFunction(UObject* SourceM
 	return Cast< UMaterialFunction >(DatasmithImporterImpl::FinalizeAsset(SourceMaterialFunction, MaterialFunctionsFolderPath, ExistingMaterialFunction, ReferencesToRemap));
 }
 
-UMaterialInterface* FDatasmithImporter::ImportMaterial(FDatasmithImportContext& ImportContext,
-	TSharedRef< IDatasmithBaseMaterialElement > MaterialElement, UMaterialInterface* ExistingMaterial)
+UMaterialInterface* FDatasmithImporter::ImportMaterial( FDatasmithImportContext& ImportContext,
+	TSharedRef< IDatasmithBaseMaterialElement > MaterialElement, UMaterialInterface* ExistingMaterial )
 {
 	TRACE_CPUPROFILER_EVENT_SCOPE(FDatasmithImporter::ImportMaterial);
 
-	UMaterialInterface* ImportedMaterial = FDatasmithMaterialImporter::CreateMaterial(ImportContext, MaterialElement, ExistingMaterial);
+	UMaterialInterface* ImportedMaterial = FDatasmithMaterialImporter::CreateMaterial( ImportContext, MaterialElement, ExistingMaterial );
 
-	if (!ImportedMaterial)
+	if ( !ImportedMaterial )
 	{
 		return nullptr;
 	}
@@ -1246,15 +1296,15 @@ UMaterialInterface* FDatasmithImporter::ImportMaterial(FDatasmithImportContext&
 	AssetImportData->AssetImportOptions = ImportContext.Options->BaseOptions.AssetOptions;
 
 	// Record requirements on mesh building for this material
-	ImportContext.AssetsContext.MaterialsRequirements.Add(MaterialElement->GetName(), FDatasmithMaterialImporter::GetMaterialRequirements(ImportedMaterial));
-	ImportContext.ImportedMaterials.Add(MaterialElement) = ImportedMaterial;
+	ImportContext.AssetsContext.MaterialsRequirements.Add( MaterialElement->GetName(), FDatasmithMaterialImporter::GetMaterialRequirements( ImportedMaterial ) );
+	ImportContext.ImportedMaterials.Add( MaterialElement ) = ImportedMaterial;
 
-	ImportMetaDataForObject(ImportContext, MaterialElement, ImportedMaterial);
+	ImportMetaDataForObject( ImportContext, MaterialElement, ImportedMaterial );
 
 	return ImportedMaterial;
 }
 
-UObject* FDatasmithImporter::FinalizeMaterial(UObject* SourceMaterial, const TCHAR* MaterialsFolderPath, UMaterialInterface* ExistingMaterial, TMap< UObject*, UObject* >* ReferencesToRemap)
+UObject* FDatasmithImporter::FinalizeMaterial( UObject* SourceMaterial, const TCHAR* MaterialsFolderPath, UMaterialInterface* ExistingMaterial, TMap< UObject*, UObject* >* ReferencesToRemap )
 {
 	TRACE_CPUPROFILER_EVENT_SCOPE(FDatasmithImporter::FinalizeMaterial);
 
@@ -1262,42 +1312,42 @@ UObject* FDatasmithImporter::FinalizeMaterial(UObject* SourceMaterial, const TCH
 
 	// Finalizing the master material might add a remapping for the instance parent property so make sure we have a remapping map available
 	TOptional< TMap< UObject*, UObject* > > ReferencesToRemapLocal;
-	if (!ReferencesToRemap)
+	if ( !ReferencesToRemap )
 	{
 		ReferencesToRemapLocal.Emplace();
 		ReferencesToRemap = &ReferencesToRemapLocal.GetValue();
 	}
 
-	if (UMaterialInstance* SourceMaterialInstance = Cast< UMaterialInstance >(SourceMaterial))
+	if ( UMaterialInstance* SourceMaterialInstance = Cast< UMaterialInstance >( SourceMaterial ) )
 	{
-		if (UMaterialInterface* SourceMaterialParent = SourceMaterialInstance->Parent)
+		if ( UMaterialInterface* SourceMaterialParent = SourceMaterialInstance->Parent )
 		{
 			FString SourceMaterialPath = SourceMaterialInstance->GetOutermost()->GetName();
 			FString SourceParentPath = SourceMaterialParent->GetOutermost()->GetName();
 
-			if (SourceParentPath.StartsWith(SourceMaterialPath))
+			if ( SourceParentPath.StartsWith( SourceMaterialPath ) )
 			{
 				// Simply finalize the source parent material.
 				// Note that the parent material will be overridden on the existing material instance
 				FString DestinationParentPath = SourceParentPath;
-				DestinationParentPath.RemoveFromStart(SourceMaterialPath);
+				DestinationParentPath.RemoveFromStart( SourceMaterialPath );
 				DestinationParentPath = MaterialsFolderPath / DestinationParentPath;
 
-				FinalizeMaterial(SourceMaterialParent, *DestinationParentPath, nullptr, ReferencesToRemap);
+				FinalizeMaterial( SourceMaterialParent, *DestinationParentPath, nullptr, ReferencesToRemap );
 			}
 		}
 	}
 
-	UMaterialEditingLibrary::DeleteAllMaterialExpressions(Cast< UMaterial >(ExistingMaterial));
+	UMaterialEditingLibrary::DeleteAllMaterialExpressions( Cast< UMaterial >( ExistingMaterial ) );
 
-	UObject* DestinationMaterial = FinalizeAsset(SourceMaterial, MaterialsFolderPath, ExistingMaterial, ReferencesToRemap);
+	UObject* DestinationMaterial = FinalizeAsset( SourceMaterial, MaterialsFolderPath, ExistingMaterial, ReferencesToRemap );
 
-	DatasmithImporterImpl::CompileMaterial(DestinationMaterial);
+	DatasmithImporterImpl::CompileMaterial( DestinationMaterial );
 
 	return DestinationMaterial;
 }
 
-void FDatasmithImporter::ImportActors(FDatasmithImportContext& ImportContext)
+void FDatasmithImporter::ImportActors( FDatasmithImportContext& ImportContext )
 {
 	TRACE_CPUPROFILER_EVENT_SCOPE(FDatasmithImporter::ImportActors);
 
@@ -1305,20 +1355,20 @@ void FDatasmithImporter::ImportActors(FDatasmithImportContext& ImportContext)
 	 * Hot fix for reimport issues UE-71655. A temporary created actor might have the same object path as the previously deleted actor.
 	 * This code below won't be needed when UE-71695 is fixed. This should be in 4.23.
 	 */
-	TArray< ADatasmithSceneActor* > SceneActors = FDatasmithImporterUtils::FindSceneActors(ImportContext.ActorsContext.FinalWorld, ImportContext.SceneAsset);
-	for (ADatasmithSceneActor* SceneActor : SceneActors)
+	TArray< ADatasmithSceneActor* > SceneActors = FDatasmithImporterUtils::FindSceneActors( ImportContext.ActorsContext.FinalWorld, ImportContext.SceneAsset );
+	for ( ADatasmithSceneActor* SceneActor : SceneActors )
 	{
-		if (!SceneActor)
+		if ( !SceneActor )
 		{
 			continue;
 		}
 
-		if (ImportContext.SceneAsset == SceneActor->Scene)
+		if ( ImportContext.SceneAsset == SceneActor->Scene )
 		{
-			for (TPair< FName, TSoftObjectPtr< AActor > >& Pair : SceneActor->RelatedActors)
+			for ( TPair< FName, TSoftObjectPtr< AActor > >& Pair : SceneActor->RelatedActors )
 			{
 				// Try to load the actor. If we can't reset the soft object ptr
-				if (!Pair.Value.LoadSynchronous())
+				if ( !Pair.Value.LoadSynchronous() )
 				{
 					Pair.Value.Reset();
 				}
@@ -1331,20 +1381,20 @@ void FDatasmithImporter::ImportActors(FDatasmithImportContext& ImportContext)
 	ADatasmithSceneActor* ImportSceneActor = ImportContext.ActorsContext.ImportSceneActor;
 
 	// Create a scene actor to import with if we don't have one
-	if (!ImportSceneActor)
+	if ( !ImportSceneActor )
 	{
 		// Create a the import scene actor for the import context
-		ImportSceneActor = FDatasmithImporterUtils::CreateImportSceneActor(ImportContext, FTransform::Identity);
+		ImportSceneActor = FDatasmithImporterUtils::CreateImportSceneActor( ImportContext, FTransform::Identity );
 	}
 
 	const int32 ActorsCount = ImportContext.Scene->GetActorsCount();
 
-	FScopedSlowTask Progress((float)ActorsCount, LOCTEXT("ImportActors", "Spawning actors..."), true, *ImportContext.Warn);
+	FScopedSlowTask Progress( (float)ActorsCount, LOCTEXT("ImportActors", "Spawning actors..."), true, *ImportContext.Warn );
 	Progress.MakeDialog(true);
 
-	if (ImportSceneActor)
+	if ( ImportSceneActor )
 	{
-		ImportContext.Hierarchy.Push(ImportSceneActor->GetRootComponent());
+		ImportContext.Hierarchy.Push( ImportSceneActor->GetRootComponent() );
 
 		for (int32 i = 0; i < ActorsCount && !ImportContext.bUserCancelled; ++i)
 		{
@@ -1352,17 +1402,17 @@ void FDatasmithImporter::ImportActors(FDatasmithImportContext& ImportContext)
 
 			TSharedPtr< IDatasmithActorElement > ActorElement = ImportContext.Scene->GetActor(i);
 
-			if (ActorElement.IsValid())
+			if ( ActorElement.IsValid() )
 			{
-				Progress.EnterProgressFrame(1.f, FText::FromString(FString::Printf(TEXT("Spawning actor %d/%d (%s) ..."), (i + 1), ActorsCount, ActorElement->GetLabel())));
+				Progress.EnterProgressFrame( 1.f, FText::FromString( FString::Printf( TEXT("Spawning actor %d/%d (%s) ..."), ( i + 1 ), ActorsCount, ActorElement->GetLabel() ) ) );
 
-				if (ActorElement->IsAComponent())
+				if ( ActorElement->IsAComponent() )
 				{
-					ImportActorAsComponent(ImportContext, ActorElement.ToSharedRef(), ImportSceneActor);
+					ImportActorAsComponent( ImportContext, ActorElement.ToSharedRef(), ImportSceneActor );
 				}
 				else
 				{
-					ImportActor(ImportContext, ActorElement.ToSharedRef());
+					ImportActor( ImportContext, ActorElement.ToSharedRef() );
 				}
 			}
 		}
@@ -1381,9 +1431,9 @@ void FDatasmithImporter::ImportActors(FDatasmithImportContext& ImportContext)
 
 			TSharedPtr< IDatasmithActorElement > ActorElement = ImportContext.Scene->GetActor(i);
 
-			if (ActorElement.IsValid() && ActorElement->IsA(EDatasmithElementType::Camera))
+			if ( ActorElement.IsValid() && ActorElement->IsA( EDatasmithElementType::Camera ) )
 			{
-				FDatasmithCameraImporter::PostImportCameraActor(StaticCastSharedRef< IDatasmithCameraActorElement >(ActorElement.ToSharedRef()), ImportContext);
+				FDatasmithCameraImporter::PostImportCameraActor( StaticCastSharedRef< IDatasmithCameraActorElement >( ActorElement.ToSharedRef() ), ImportContext );
 			}
 		}
 
@@ -1393,70 +1443,70 @@ void FDatasmithImporter::ImportActors(FDatasmithImportContext& ImportContext)
 	}
 
 	// Sky
-	if (ImportContext.Scene->GetUsePhysicalSky())
+	if( ImportContext.Scene->GetUsePhysicalSky() )
 	{
 		AActor* SkyActor = FDatasmithLightImporter::CreatePhysicalSky(ImportContext);
 	}
 
-	if (ImportContext.bUserCancelled)
+	if ( ImportContext.bUserCancelled )
 	{
-		DatasmithImporterImpl::DeleteImportSceneActorIfNeeded(ImportContext.ActorsContext, true);
+		DatasmithImporterImpl::DeleteImportSceneActorIfNeeded(ImportContext.ActorsContext, true );
 	}
 }
 
-AActor* FDatasmithImporter::ImportActor(FDatasmithImportContext& ImportContext, const TSharedRef< IDatasmithActorElement >& ActorElement)
+AActor* FDatasmithImporter::ImportActor( FDatasmithImportContext& ImportContext, const TSharedRef< IDatasmithActorElement >& ActorElement )
 {
 	TRACE_CPUPROFILER_EVENT_SCOPE(FDatasmithImporter::ImportActor);
 
 	AActor* ImportedActor = nullptr;
 	if (ActorElement->IsA(EDatasmithElementType::HierarchicalInstanceStaticMesh))
 	{
-		TSharedRef< IDatasmithHierarchicalInstancedStaticMeshActorElement > HISMActorElement = StaticCastSharedRef< IDatasmithHierarchicalInstancedStaticMeshActorElement >(ActorElement);
-		ImportedActor = FDatasmithActorImporter::ImportHierarchicalInstancedStaticMeshAsActor(ImportContext, HISMActorElement);
+		TSharedRef< IDatasmithHierarchicalInstancedStaticMeshActorElement > HISMActorElement = StaticCastSharedRef< IDatasmithHierarchicalInstancedStaticMeshActorElement >( ActorElement );
+		ImportedActor =  FDatasmithActorImporter::ImportHierarchicalInstancedStaticMeshAsActor( ImportContext, HISMActorElement );
 	}
 	else if (ActorElement->IsA(EDatasmithElementType::StaticMeshActor))
 	{
-		TSharedRef< IDatasmithMeshActorElement > MeshActorElement = StaticCastSharedRef< IDatasmithMeshActorElement >(ActorElement);
-		ImportedActor = FDatasmithActorImporter::ImportStaticMeshActor(ImportContext, MeshActorElement);
+		TSharedRef< IDatasmithMeshActorElement > MeshActorElement = StaticCastSharedRef< IDatasmithMeshActorElement >( ActorElement );
+		ImportedActor = FDatasmithActorImporter::ImportStaticMeshActor( ImportContext, MeshActorElement );
 	}
 	else if (ActorElement->IsA(EDatasmithElementType::EnvironmentLight))
 	{
-		ImportedActor = FDatasmithActorImporter::ImportEnvironment(ImportContext, StaticCastSharedRef< IDatasmithEnvironmentElement >(ActorElement));
+		ImportedActor = FDatasmithActorImporter::ImportEnvironment( ImportContext, StaticCastSharedRef< IDatasmithEnvironmentElement >( ActorElement ) );
 	}
 	else if (ActorElement->IsA(EDatasmithElementType::Light))
 	{
-		ImportedActor = FDatasmithActorImporter::ImportLightActor(ImportContext, StaticCastSharedRef< IDatasmithLightActorElement >(ActorElement));
+		ImportedActor = FDatasmithActorImporter::ImportLightActor( ImportContext, StaticCastSharedRef< IDatasmithLightActorElement >(ActorElement) );
 	}
 	else if (ActorElement->IsA(EDatasmithElementType::Camera))
 	{
-		ImportedActor = FDatasmithActorImporter::ImportCameraActor(ImportContext, StaticCastSharedRef< IDatasmithCameraActorElement >(ActorElement));
+		ImportedActor = FDatasmithActorImporter::ImportCameraActor( ImportContext, StaticCastSharedRef< IDatasmithCameraActorElement >(ActorElement) );
 	}
 	else if (ActorElement->IsA(EDatasmithElementType::CustomActor))
 	{
-		ImportedActor = FDatasmithActorImporter::ImportCustomActor(ImportContext, StaticCastSharedRef< IDatasmithCustomActorElement >(ActorElement));
+		ImportedActor = FDatasmithActorImporter::ImportCustomActor( ImportContext, StaticCastSharedRef< IDatasmithCustomActorElement >(ActorElement) );
 	}
 	else if (ActorElement->IsA(EDatasmithElementType::Landscape))
 	{
-		ImportedActor = FDatasmithActorImporter::ImportLandscapeActor(ImportContext, StaticCastSharedRef< IDatasmithLandscapeElement >(ActorElement));
+		ImportedActor = FDatasmithActorImporter::ImportLandscapeActor( ImportContext, StaticCastSharedRef< IDatasmithLandscapeElement >(ActorElement) );
 	}
 	else if (ActorElement->IsA(EDatasmithElementType::PostProcessVolume))
 	{
-		ImportedActor = FDatasmithPostProcessImporter::ImportPostProcessVolume(StaticCastSharedRef< IDatasmithPostProcessVolumeElement >(ActorElement), ImportContext, ImportContext.Options->OtherActorImportPolicy);
+		ImportedActor = FDatasmithPostProcessImporter::ImportPostProcessVolume( StaticCastSharedRef< IDatasmithPostProcessVolumeElement >( ActorElement ), ImportContext, ImportContext.Options->OtherActorImportPolicy );
 	}
 	else
 	{
-		ImportedActor = FDatasmithActorImporter::ImportBaseActor(ImportContext, ActorElement);
+		ImportedActor = FDatasmithActorImporter::ImportBaseActor( ImportContext, ActorElement );
 	}
 
 
-	if (ImportedActor) // It's possible that we didn't import an actor (ie: the user doesn't want to import the cameras), in that case, we'll skip it in the hierarchy
+	if ( ImportedActor ) // It's possible that we didn't import an actor (ie: the user doesn't want to import the cameras), in that case, we'll skip it in the hierarchy
 	{
-		ImportContext.Hierarchy.Push(ImportedActor->GetRootComponent());
+		ImportContext.Hierarchy.Push( ImportedActor->GetRootComponent() );
 		ImportMetaDataForObject(ImportContext, ActorElement, ImportedActor);
 	}
 	else
 	{
-		ImportContext.ActorsContext.NonImportedDatasmithActors.Add(ActorElement->GetName());
+		ImportContext.ActorsContext.NonImportedDatasmithActors.Add( ActorElement->GetName() );
 	}
 
 	for (int32 i = 0; i < ActorElement->GetChildrenCount() && !ImportContext.bUserCancelled; ++i)
@@ -1465,20 +1515,20 @@ AActor* FDatasmithImporter::ImportActor(FDatasmithImportContext& ImportContext,
 
 		const TSharedPtr< IDatasmithActorElement >& ChildActorElement = ActorElement->GetChild(i);
 
-		if (ChildActorElement.IsValid())
+		if ( ChildActorElement.IsValid() )
 		{
-			if (ImportContext.Options->HierarchyHandling == EDatasmithImportHierarchy::UseMultipleActors && !ChildActorElement->IsAComponent())
+			if ( ImportContext.Options->HierarchyHandling == EDatasmithImportHierarchy::UseMultipleActors && !ChildActorElement->IsAComponent() )
 			{
-				ImportActor(ImportContext, ChildActorElement.ToSharedRef());
+				ImportActor( ImportContext, ChildActorElement.ToSharedRef() );
 			}
-			else if (ImportedActor) // Don't import the components of an actor that we didn't import
+			else if ( ImportedActor ) // Don't import the components of an actor that we didn't import
 			{
-				ImportActorAsComponent(ImportContext, ChildActorElement.ToSharedRef(), ImportedActor);
+				ImportActorAsComponent( ImportContext, ChildActorElement.ToSharedRef(), ImportedActor );
 			}
 		}
 	}
 
-	if (ImportedActor)
+	if ( ImportedActor )
 	{
 		ImportContext.Hierarchy.Pop();
 	}
@@ -1554,52 +1604,52 @@ void FDatasmithImporter::ImportActorAsComponent(FDatasmithImportContext& ImportC
 	}
 }
 
-void FDatasmithImporter::FinalizeActors(FDatasmithImportContext& ImportContext, TMap< UObject*, UObject* >* AssetReferencesToRemap)
+void FDatasmithImporter::FinalizeActors( FDatasmithImportContext& ImportContext, TMap< UObject*, UObject* >* AssetReferencesToRemap )
 {
 	TRACE_CPUPROFILER_EVENT_SCOPE(FDatasmithImporter::FinalizeActors);
 
 	using namespace DatasmithImporterImpl;
 
-	if (!ImportContext.bUserCancelled)
+	if ( !ImportContext.bUserCancelled )
 	{
 		// Ensure a proper setup for the finalize of the actors
 		ADatasmithSceneActor*& ImportSceneActor = ImportContext.ActorsContext.ImportSceneActor;
 		TSet< ADatasmithSceneActor* >& FinalSceneActors = ImportContext.ActorsContext.FinalSceneActors;
 
-		if (!ImportContext.ActorsContext.FinalWorld)
+		if ( !ImportContext.ActorsContext.FinalWorld )
 		{
 			ImportContext.ActorsContext.FinalWorld = ImportContext.ActorsContext.ImportWorld;
 		}
-		else if (!ImportContext.bIsAReimport && ImportSceneActor)
+		else if ( !ImportContext.bIsAReimport && ImportSceneActor )
 		{
-			//Create a new datasmith scene actor in the final level
-			FActorSpawnParameters SpawnParameters;
-			SpawnParameters.Template = ImportSceneActor;
-			ADatasmithSceneActor* DestinationSceneActor = Cast< ADatasmithSceneActor >(ImportContext.ActorsContext.FinalWorld->SpawnActor< ADatasmithSceneActor >(SpawnParameters));
+				//Create a new datasmith scene actor in the final level
+				FActorSpawnParameters SpawnParameters;
+				SpawnParameters.Template = ImportSceneActor;
+				ADatasmithSceneActor* DestinationSceneActor = Cast< ADatasmithSceneActor >(ImportContext.ActorsContext.FinalWorld->SpawnActor< ADatasmithSceneActor >(SpawnParameters));
 
-			// Name new destination ADatasmithSceneActor to the DatasmithScene's name
-			DestinationSceneActor->SetActorLabel(ImportContext.Scene->GetName());
-			DestinationSceneActor->MarkPackageDirty();
-			DestinationSceneActor->RelatedActors.Reset();
+				// Name new destination ADatasmithSceneActor to the DatasmithScene's name
+				DestinationSceneActor->SetActorLabel( ImportContext.Scene->GetName() );
+				DestinationSceneActor->MarkPackageDirty();
+				DestinationSceneActor->RelatedActors.Reset();
 
-			FinalSceneActors.Empty(1);
-			FinalSceneActors.Add(DestinationSceneActor);
+				FinalSceneActors.Empty( 1 );
+				FinalSceneActors.Add( DestinationSceneActor );
 		}
 
-		if (FinalSceneActors.Num() == 0)
+		if( FinalSceneActors.Num() == 0 )
 		{
-			if (ImportContext.bIsAReimport)
+			if ( ImportContext.bIsAReimport )
 			{
-				FinalSceneActors.Append(FDatasmithImporterUtils::FindSceneActors(ImportContext.ActorsContext.FinalWorld, ImportContext.SceneAsset));
-				FinalSceneActors.Remove(ImportSceneActor);
+				FinalSceneActors.Append( FDatasmithImporterUtils::FindSceneActors( ImportContext.ActorsContext.FinalWorld, ImportContext.SceneAsset ) );
+				FinalSceneActors.Remove( ImportSceneActor );
 			}
 			else
 			{
-				FinalSceneActors.Add(ImportSceneActor);
+				FinalSceneActors.Add( ImportSceneActor );
 			}
 		}
 
-		for (AActor* Actor : FinalSceneActors)
+		for ( AActor* Actor : FinalSceneActors )
 		{
 			check(Actor->GetWorld() == ImportContext.ActorsContext.FinalWorld);
 		}
@@ -1610,14 +1660,14 @@ void FDatasmithImporter::FinalizeActors(FDatasmithImportContext& ImportContext,
 		TSet< FName > LayersUsedByActors;
 		const bool bShouldSpawnNonExistingActors = !ImportContext.bIsAReimport || ImportContext.Options->ReimportOptions.bRespawnDeletedActors;
 
-		for (ADatasmithSceneActor* DestinationSceneActor : FinalSceneActors)
+		for ( ADatasmithSceneActor* DestinationSceneActor : FinalSceneActors )
 		{
-			if (!DestinationSceneActor)
+			if ( !DestinationSceneActor )
 			{
 				continue;
 			}
 
-			if (ImportSceneActor->Scene != DestinationSceneActor->Scene)
+			if ( ImportSceneActor->Scene != DestinationSceneActor->Scene )
 			{
 				continue;
 			}
@@ -1627,55 +1677,55 @@ void FDatasmithImporter::FinalizeActors(FDatasmithImportContext& ImportContext,
 
 			ImportContext.ActorsContext.CurrentTargetedScene = DestinationSceneActor;
 
-			if (ImportSceneActor != DestinationSceneActor)
+			if ( ImportSceneActor != DestinationSceneActor )
 			{
 				// Before we delete the non imported actors, remove the old actor labels from the unique name provider
 				// as we don't care if the source labels clash with labels from actors that will be deleted or replaced on reimport
 				for (const TPair< FName, TSoftObjectPtr< AActor > >& ActorPair : DestinationSceneActor->RelatedActors)
 				{
-					if (AActor* DestActor = ActorPair.Value.Get())
+					if ( AActor* DestActor = ActorPair.Value.Get() )
 					{
 						ImportContext.ActorsContext.UniqueNameProvider.RemoveExistingName(DestActor->GetActorLabel());
 					}
 				}
 
-				FDatasmithImporterUtils::DeleteNonImportedDatasmithElementFromSceneActor(*ImportSceneActor, *DestinationSceneActor, ImportContext.ActorsContext.NonImportedDatasmithActors);
+				FDatasmithImporterUtils::DeleteNonImportedDatasmithElementFromSceneActor( *ImportSceneActor, *DestinationSceneActor, ImportContext.ActorsContext.NonImportedDatasmithActors );
 			}
 
 			// Add Actor info to the remap info
 			TMap< UObject*, UObject* > PerSceneActorReferencesToRemap = AssetReferencesToRemap ? *AssetReferencesToRemap : TMap< UObject*, UObject* >();
-			PerSceneActorReferencesToRemap.Add(ImportSceneActor) = DestinationSceneActor;
-			PerSceneActorReferencesToRemap.Add(ImportSceneActor->GetRootComponent()) = DestinationSceneActor->GetRootComponent();
+			PerSceneActorReferencesToRemap.Add( ImportSceneActor ) = DestinationSceneActor;
+			PerSceneActorReferencesToRemap.Add( ImportSceneActor->GetRootComponent() ) = DestinationSceneActor->GetRootComponent();
 
 			// #ueent_todo order of actors matters for ReferencesFix + re-parenting
-			for (const TPair< FName, TSoftObjectPtr< AActor > >& SourceActorPair : ImportSceneActor->RelatedActors)
+			for ( const TPair< FName, TSoftObjectPtr< AActor > >& SourceActorPair : ImportSceneActor->RelatedActors )
 			{
 				AActor* SourceActor = SourceActorPair.Value.Get();
-				if (SourceActor == nullptr)
+				if ( SourceActor == nullptr )
 				{
 					continue;
 				}
 
-				const bool bActorIsRelatedToDestionScene = DestinationSceneActor->RelatedActors.Contains(SourceActorPair.Key);
-				TSoftObjectPtr< AActor >& ExistingActorPtr = DestinationSceneActor->RelatedActors.FindOrAdd(SourceActorPair.Key);
-				const bool bShouldFinalizeActor = bShouldSpawnNonExistingActors || !bActorIsRelatedToDestionScene || (ExistingActorPtr.Get() && !ExistingActorPtr.Get()->IsPendingKillPending());
+				const bool bActorIsRelatedToDestionScene = DestinationSceneActor->RelatedActors.Contains( SourceActorPair.Key );
+				TSoftObjectPtr< AActor >& ExistingActorPtr = DestinationSceneActor->RelatedActors.FindOrAdd( SourceActorPair.Key );
+				const bool bShouldFinalizeActor = bShouldSpawnNonExistingActors || !bActorIsRelatedToDestionScene || ( ExistingActorPtr.Get() && !ExistingActorPtr.Get()->IsPendingKillPending() );
 
-				if (bShouldFinalizeActor)
+				if ( bShouldFinalizeActor )
 				{
 					// Remember the original source path as FinalizeActor may set SourceActor's label, which apparently can also change its Name and package path
 					FSoftObjectPath OriginalSourcePath = FSoftObjectPath(SourceActor);
-					AActor* DestinationActor = FinalizeActor(ImportContext, *SourceActor, ExistingActorPtr.Get(), PerSceneActorReferencesToRemap);
+					AActor* DestinationActor = FinalizeActor( ImportContext, *SourceActor, ExistingActorPtr.Get(), PerSceneActorReferencesToRemap );
 					RenamedActorsMap.Add(OriginalSourcePath, FSoftObjectPath(DestinationActor));
 					LayersUsedByActors.Append(DestinationActor->Layers);
 					ExistingActorPtr = DestinationActor;
 				}
 			}
 
-			for (const TPair< FName, TSoftObjectPtr< AActor > >& DestinationActorPair : DestinationSceneActor->RelatedActors)
+			for (  const TPair< FName, TSoftObjectPtr< AActor > >& DestinationActorPair : DestinationSceneActor->RelatedActors )
 			{
-				if (AActor* Actor = DestinationActorPair.Value.Get())
+				if ( AActor* Actor = DestinationActorPair.Value.Get() )
 				{
-					FixReferencesForObject(Actor, PerSceneActorReferencesToRemap);
+					FixReferencesForObject( Actor, PerSceneActorReferencesToRemap );
 				}
 			}
 
@@ -1684,74 +1734,74 @@ void FDatasmithImporter::FinalizeActors(FDatasmithImportContext& ImportContext,
 		}
 
 		// Add the missing layers to the final world
-		FDatasmithImporterUtils::AddUniqueLayersToWorld(ImportContext.ActorsContext.FinalWorld, LayersUsedByActors);
+		FDatasmithImporterUtils::AddUniqueLayersToWorld( ImportContext.ActorsContext.FinalWorld, LayersUsedByActors );
 
 		// Fixed the soft object paths that were pointing to our pre-finalized actors.
 		TArray< UPackage* > PackagesToFix;
 
-		for (const TPair< FName, TSoftObjectPtr< ULevelSequence > >& LevelSequencePair : ImportContext.SceneAsset->LevelSequences)
+		for ( const TPair< FName, TSoftObjectPtr< ULevelSequence > >& LevelSequencePair : ImportContext.SceneAsset->LevelSequences )
 		{
 			if (LevelSequencePair.Value)
 			{
-				PackagesToFix.Add(LevelSequencePair.Value->GetOutermost());
+				PackagesToFix.Add( LevelSequencePair.Value->GetOutermost() );
 			}
 		}
 
-		for (const TPair< FName, TSoftObjectPtr< ULevelVariantSets > >& LevelVariantSetsPair : ImportContext.SceneAsset->LevelVariantSets)
+		for ( const TPair< FName, TSoftObjectPtr< ULevelVariantSets > >& LevelVariantSetsPair : ImportContext.SceneAsset->LevelVariantSets )
 		{
 			if (LevelVariantSetsPair.Value)
 			{
-				PackagesToFix.Add(LevelVariantSetsPair.Value->GetOutermost());
+				PackagesToFix.Add( LevelVariantSetsPair.Value->GetOutermost() );
 			}
 		}
 
 		FAssetToolsModule& AssetToolsModule = FModuleManager::GetModuleChecked< FAssetToolsModule >(TEXT("AssetTools"));
-		AssetToolsModule.Get().RenameReferencingSoftObjectPaths(PackagesToFix, RenamedActorsMap);
+		AssetToolsModule.Get().RenameReferencingSoftObjectPaths( PackagesToFix, RenamedActorsMap );
 	}
 
-	DeleteImportSceneActorIfNeeded(ImportContext.ActorsContext);
+	DeleteImportSceneActorIfNeeded( ImportContext.ActorsContext );
 
 	// Ensure layer visibility is properly updated for new actors associated with existing layers
 	ULayersSubsystem* LayersSubsystem = GEditor->GetEditorSubsystem<ULayersSubsystem>();
-	LayersSubsystem->UpdateAllActorsVisibility(false, true);
+	LayersSubsystem->UpdateAllActorsVisibility( false, true);
 
 	GEngine->BroadcastLevelActorListChanged();
 }
 
-AActor* FDatasmithImporter::FinalizeActor(FDatasmithImportContext& ImportContext, AActor& SourceActor, AActor* ExistingActor, TMap< UObject*, UObject* >& ReferencesToRemap)
+AActor* FDatasmithImporter::FinalizeActor( FDatasmithImportContext& ImportContext, AActor& SourceActor, AActor* ExistingActor, TMap< UObject*, UObject* >& ReferencesToRemap )
 {
 	TRACE_CPUPROFILER_EVENT_SCOPE(FDatasmithImporter::FinalizeActor);
 
 	using namespace DatasmithImporterImpl;
 
 	// If the existing actor is not of the same class we destroy it
-	if (ExistingActor && ExistingActor->GetClass() != SourceActor.GetClass())
+	if ( ExistingActor && ExistingActor->GetClass() != SourceActor.GetClass() )
 	{
-		FDatasmithImporterUtils::DeleteActor(*ExistingActor);
+		FDatasmithImporterUtils::DeleteActor( *ExistingActor );
 		ExistingActor = nullptr;
 	}
 
 	// Backup hierarchy
 	TArray< AActor* > Children;
-	if (ExistingActor)
+	if ( ExistingActor )
 	{
-		ExistingActor->GetAttachedActors(Children);
+		ExistingActor->GetAttachedActors( Children );
 
 		// In order to allow modification on components owned by ExistingActor, unregister all of them
-		ExistingActor->UnregisterAllComponents( /* bForReregister = */true);
+		ExistingActor->UnregisterAllComponents( /* bForReregister = */true );
 	}
 
 	AActor* DestinationActor = ExistingActor;
-	if (!DestinationActor)
+	if ( !DestinationActor )
 	{
-		DestinationActor = ImportContext.ActorsContext.FinalWorld->SpawnActor(SourceActor.GetClass());
+		DestinationActor = ImportContext.ActorsContext.FinalWorld->SpawnActor( SourceActor.GetClass() );
 	}
 
 	// Update label to match the source actor's
-	DestinationActor->SetActorLabel(ImportContext.ActorsContext.UniqueNameProvider.GenerateUniqueName(SourceActor.GetActorLabel()));
+	DestinationActor->SetActorLabel( ImportContext.ActorsContext.UniqueNameProvider.GenerateUniqueName( SourceActor.GetActorLabel() ) );
 
-	check(DestinationActor);
-	ReferencesToRemap.Add(&SourceActor) = DestinationActor;
+	check( DestinationActor );
+	ReferencesToRemap.Add( &SourceActor ) = DestinationActor;
 
 	TArray< FMigratedTemplatePairType > MigratedTemplates = MigrateTemplates(
 		&SourceActor,
@@ -1763,35 +1813,35 @@ AActor* FDatasmithImporter::FinalizeActor(FDatasmithImportContext& ImportContext
 	// Copy actor data
 	{
 		TArray< uint8 > Bytes;
-		FActorWriter ObjectWriter(&SourceActor, Bytes);
-		FObjectReader ObjectReader(DestinationActor, Bytes);
+		FActorWriter ObjectWriter( &SourceActor, Bytes );
+		FObjectReader ObjectReader( DestinationActor, Bytes );
 	}
 
-	FixReferencesForObject(DestinationActor, ReferencesToRemap);
+	FixReferencesForObject( DestinationActor, ReferencesToRemap );
 
-	DatasmithImporterImpl::FinalizeComponents(ImportContext, SourceActor, *DestinationActor, ReferencesToRemap);
+	DatasmithImporterImpl::FinalizeComponents( ImportContext, SourceActor, *DestinationActor, ReferencesToRemap );
 
 	// The templates for the actor need to be applied after the components were created.
-	ApplyMigratedTemplates(MigratedTemplates, DestinationActor);
+	ApplyMigratedTemplates( MigratedTemplates, DestinationActor );
 
 	// Restore hierarchy
-	for (AActor* Child : Children)
+	for ( AActor* Child : Children )
 	{
-		Child->AttachToActor(DestinationActor, FAttachmentTransformRules::KeepWorldTransform);
+		Child->AttachToActor( DestinationActor, FAttachmentTransformRules::KeepWorldTransform );
 	}
 
 	// Hotfix for UE-69555
 	TInlineComponentArray<UHierarchicalInstancedStaticMeshComponent*> HierarchicalInstancedStaticMeshComponents;
 	DestinationActor->GetComponents(HierarchicalInstancedStaticMeshComponents);
-	for (UHierarchicalInstancedStaticMeshComponent* HierarchicalInstancedStaticMeshComponent : HierarchicalInstancedStaticMeshComponents)
+	for (UHierarchicalInstancedStaticMeshComponent* HierarchicalInstancedStaticMeshComponent : HierarchicalInstancedStaticMeshComponents )
 	{
-		HierarchicalInstancedStaticMeshComponent->BuildTreeIfOutdated(true, true);
+		HierarchicalInstancedStaticMeshComponent->BuildTreeIfOutdated( true, true );
 	}
 
-	if (ALandscape* Landscape = Cast< ALandscape >(DestinationActor))
+	if ( ALandscape* Landscape = Cast< ALandscape >( DestinationActor ) )
 	{
-		FPropertyChangedEvent MaterialPropertyChangedEvent(FindFieldChecked< UProperty >(Landscape->GetClass(), FName("LandscapeMaterial")));
-		Landscape->PostEditChangeProperty(MaterialPropertyChangedEvent);
+		FPropertyChangedEvent MaterialPropertyChangedEvent( FindFieldChecked< UProperty >( Landscape->GetClass(), FName("LandscapeMaterial") ) );
+		Landscape->PostEditChangeProperty( MaterialPropertyChangedEvent );
 	}
 
 	FQuat PreviousRotation = DestinationActor->GetRootComponent()->GetRelativeTransform().GetRotation();
@@ -1824,10 +1874,10 @@ AActor* FDatasmithImporter::FinalizeActor(FDatasmithImportContext& ImportContext
 	return DestinationActor;
 }
 
-void FDatasmithImporter::ImportLevelSequences(FDatasmithImportContext& ImportContext)
+void FDatasmithImporter::ImportLevelSequences( FDatasmithImportContext& ImportContext )
 {
 	const int32 SequencesCount = ImportContext.FilteredScene->GetLevelSequencesCount();
-	if (!ImportContext.Options->BaseOptions.CanIncludeAnimation() || !ImportContext.Options->BaseOptions.bIncludeAnimation || SequencesCount == 0)
+	if ( !ImportContext.Options->BaseOptions.CanIncludeAnimation() || !ImportContext.Options->BaseOptions.bIncludeAnimation || SequencesCount == 0 )
 	{
 		return;
 	}
@@ -1840,12 +1890,12 @@ void FDatasmithImporter::ImportLevelSequences(FDatasmithImportContext& ImportCon
 	// iteratively resolving all dependencies
 	TArray<TSharedPtr<IDatasmithLevelSequenceElement>> SequencesToImport;
 	SequencesToImport.Reserve(SequencesCount);
-	for (int32 SequenceIndex = 0; SequenceIndex < SequencesCount && !ImportContext.bUserCancelled; ++SequenceIndex)
+	for ( int32 SequenceIndex = 0; SequenceIndex < SequencesCount && !ImportContext.bUserCancelled; ++SequenceIndex )
 	{
 		ImportContext.bUserCancelled |= ImportContext.Warn->ReceivedUserCancel();
 
-		TSharedPtr< IDatasmithLevelSequenceElement > SequenceElement = ImportContext.FilteredScene->GetLevelSequence(SequenceIndex);
-		if (!SequenceElement)
+		TSharedPtr< IDatasmithLevelSequenceElement > SequenceElement = ImportContext.FilteredScene->GetLevelSequence( SequenceIndex );
+		if ( !SequenceElement )
 		{
 			continue;
 		}
@@ -1860,7 +1910,7 @@ void FDatasmithImporter::ImportLevelSequences(FDatasmithImportContext& ImportCon
 	for (int32 IterationCounter = 0; IterationCounter < HardLoopCounter && !ImportContext.bUserCancelled; ++IterationCounter)
 	{
 		// Scan remaining sequences and import the ones we can, removing from this array
-		for (int32 SequenceIndex = SequencesToImport.Num() - 1; SequenceIndex >= 0 && !ImportContext.bUserCancelled; --SequenceIndex)
+		for ( int32 SequenceIndex = SequencesToImport.Num() - 1; SequenceIndex >= 0 && !ImportContext.bUserCancelled; --SequenceIndex )
 		{
 			ImportContext.bUserCancelled |= ImportContext.Warn->ReceivedUserCancel();
 
@@ -1871,27 +1921,27 @@ void FDatasmithImporter::ImportLevelSequences(FDatasmithImportContext& ImportCon
 				continue;
 			}
 
-			ULevelSequence* ExistingLevelSequence = nullptr;
-			if (ImportContext.SceneAsset)
-			{
-				TSoftObjectPtr< ULevelSequence >* ExistingLevelSequencePtr = ImportContext.SceneAsset->LevelSequences.Find(SequenceElement->GetName());
+		ULevelSequence* ExistingLevelSequence = nullptr;
+		if ( ImportContext.SceneAsset )
+		{
+			TSoftObjectPtr< ULevelSequence >* ExistingLevelSequencePtr = ImportContext.SceneAsset->LevelSequences.Find( SequenceElement->GetName() );
 
-				if (ExistingLevelSequencePtr)
-				{
-					ExistingLevelSequence = ExistingLevelSequencePtr->LoadSynchronous();
-				}
+			if ( ExistingLevelSequencePtr )
+			{
+				ExistingLevelSequence = ExistingLevelSequencePtr->LoadSynchronous();
 			}
+		}
 
-			FString SequenceName = ObjectTools::SanitizeObjectName(SequenceElement->GetName());
-			Progress.EnterProgressFrame(1.f, FText::FromString(FString::Printf(TEXT("Importing level sequence %d/%d (%s) ..."), NumImported + 1, HardLoopCounter, *SequenceName)));
+			FString SequenceName = ObjectTools::SanitizeObjectName( SequenceElement->GetName() );
+			Progress.EnterProgressFrame( 1.f, FText::FromString( FString::Printf( TEXT("Importing level sequence %d/%d (%s) ..."), NumImported + 1, HardLoopCounter, *SequenceName ) ) );
 
-			ULevelSequence*& ImportedLevelSequence = ImportContext.ImportedLevelSequences.FindOrAdd(SequenceElement.ToSharedRef());
-			if (ImportContext.SceneTranslator)
-			{
-				FDatasmithLevelSequencePayload LevelSequencePayload;
-				ImportContext.SceneTranslator->LoadLevelSequence(SequenceElement.ToSharedRef(), LevelSequencePayload);
-			}
-			ImportedLevelSequence = FDatasmithLevelSequenceImporter::ImportLevelSequence(SequenceElement.ToSharedRef(), ImportContext, ExistingLevelSequence);
+		ULevelSequence*& ImportedLevelSequence = ImportContext.ImportedLevelSequences.FindOrAdd( SequenceElement.ToSharedRef() );
+		if (ImportContext.SceneTranslator)
+		{
+			FDatasmithLevelSequencePayload LevelSequencePayload;
+			ImportContext.SceneTranslator->LoadLevelSequence(SequenceElement.ToSharedRef(), LevelSequencePayload);
+		}
+		ImportedLevelSequence = FDatasmithLevelSequenceImporter::ImportLevelSequence( SequenceElement.ToSharedRef(), ImportContext, ExistingLevelSequence );
 
 			SequencesToImport.RemoveAt(SequenceIndex);
 			++NumImported;
@@ -1909,7 +1959,7 @@ void FDatasmithImporter::ImportLevelSequences(FDatasmithImportContext& ImportCon
 	if (SequencesToImport.Num() > 0)
 	{
 		FString ErrorMessage = LOCTEXT("FailedToImport", "Failed to import some animation sequences:\n").ToString();
-		for (TSharedPtr<IDatasmithLevelSequenceElement>& Sequence : SequencesToImport)
+		for (TSharedPtr<IDatasmithLevelSequenceElement>& Sequence: SequencesToImport)
 		{
 			ErrorMessage += FString(TEXT("\t")) + Sequence->GetName() + TEXT("\n");
 		}
@@ -1920,15 +1970,15 @@ void FDatasmithImporter::ImportLevelSequences(FDatasmithImportContext& ImportCon
 	ImportContext.AssetsContext.LevelSequencesImportPackage.Reset();
 }
 
-ULevelSequence* FDatasmithImporter::FinalizeLevelSequence(ULevelSequence* SourceLevelSequence, const TCHAR* AnimationsFolderPath, ULevelSequence* ExistingLevelSequence)
+ULevelSequence* FDatasmithImporter::FinalizeLevelSequence( ULevelSequence* SourceLevelSequence, const TCHAR* AnimationsFolderPath, ULevelSequence* ExistingLevelSequence )
 {
-	return Cast< ULevelSequence >(DatasmithImporterImpl::PublicizeAsset(SourceLevelSequence, AnimationsFolderPath, ExistingLevelSequence));
+	return Cast< ULevelSequence >( DatasmithImporterImpl::PublicizeAsset( SourceLevelSequence, AnimationsFolderPath, ExistingLevelSequence ) );
 }
 
-void FDatasmithImporter::ImportLevelVariantSets(FDatasmithImportContext& ImportContext)
+void FDatasmithImporter::ImportLevelVariantSets( FDatasmithImportContext& ImportContext )
 {
 	const int32 LevelVariantSetsCount = ImportContext.FilteredScene->GetLevelVariantSetsCount();
-	if (LevelVariantSetsCount == 0)
+	if ( LevelVariantSetsCount == 0 )
 	{
 		return;
 	}
@@ -1936,46 +1986,46 @@ void FDatasmithImporter::ImportLevelVariantSets(FDatasmithImportContext& ImportC
 	FScopedSlowTask Progress( (float)LevelVariantSetsCount, LOCTEXT("ImportingLevelVariantSets", "Importing Level Variant Sets..."), true, *ImportContext.Warn );
 	Progress.MakeDialog(true);
 
-	for (int32 LevelVariantSetIndex = 0; LevelVariantSetIndex < LevelVariantSetsCount && !ImportContext.bUserCancelled; ++LevelVariantSetIndex)
+	for ( int32 LevelVariantSetIndex = 0; LevelVariantSetIndex < LevelVariantSetsCount && !ImportContext.bUserCancelled; ++LevelVariantSetIndex )
 	{
 		ImportContext.bUserCancelled |= ImportContext.Warn->ReceivedUserCancel();
 
-		TSharedPtr< IDatasmithLevelVariantSetsElement > LevelVariantSetsElement = ImportContext.FilteredScene->GetLevelVariantSets(LevelVariantSetIndex);
-		if (!LevelVariantSetsElement)
+		TSharedPtr< IDatasmithLevelVariantSetsElement > LevelVariantSetsElement = ImportContext.FilteredScene->GetLevelVariantSets( LevelVariantSetIndex );
+		if ( !LevelVariantSetsElement )
 		{
 			continue;
 		}
 
 		ULevelVariantSets* ExistingLevelVariantSets = nullptr;
-		if (ImportContext.SceneAsset)
+		if ( ImportContext.SceneAsset )
 		{
-			TSoftObjectPtr< ULevelVariantSets >* ExistingLevelVariantSetsPtr = ImportContext.SceneAsset->LevelVariantSets.Find(LevelVariantSetsElement->GetName());
+			TSoftObjectPtr< ULevelVariantSets >* ExistingLevelVariantSetsPtr = ImportContext.SceneAsset->LevelVariantSets.Find( LevelVariantSetsElement->GetName() );
 
-			if (ExistingLevelVariantSetsPtr)
+			if ( ExistingLevelVariantSetsPtr )
 			{
 				ExistingLevelVariantSets = ExistingLevelVariantSetsPtr->LoadSynchronous();
 			}
 		}
 
-		FString LevelVariantSetsName = ObjectTools::SanitizeObjectName(LevelVariantSetsElement->GetName());
-		Progress.EnterProgressFrame(1.f, FText::FromString(FString::Printf(TEXT("Importing level variant sets %d/%d (%s) ..."), LevelVariantSetIndex + 1, LevelVariantSetsCount, *LevelVariantSetsName)));
+		FString LevelVariantSetsName = ObjectTools::SanitizeObjectName( LevelVariantSetsElement->GetName() );
+		Progress.EnterProgressFrame( 1.f, FText::FromString( FString::Printf( TEXT("Importing level variant sets %d/%d (%s) ..."), LevelVariantSetIndex + 1, LevelVariantSetsCount, *LevelVariantSetsName ) ) );
 
-		ULevelVariantSets*& ImportedLevelVariantSets = ImportContext.ImportedLevelVariantSets.FindOrAdd(LevelVariantSetsElement.ToSharedRef());
-		ImportedLevelVariantSets = FDatasmithLevelVariantSetsImporter::ImportLevelVariantSets(LevelVariantSetsElement.ToSharedRef(), ImportContext, ExistingLevelVariantSets);
+		ULevelVariantSets*& ImportedLevelVariantSets = ImportContext.ImportedLevelVariantSets.FindOrAdd( LevelVariantSetsElement.ToSharedRef() );
+		ImportedLevelVariantSets = FDatasmithLevelVariantSetsImporter::ImportLevelVariantSets( LevelVariantSetsElement.ToSharedRef(), ImportContext, ExistingLevelVariantSets );
 	}
 
 	// Assets have been imported and moved out of their import packages, clear them so that we don't look for them in there anymore
 	ImportContext.AssetsContext.LevelVariantSetsImportPackage.Reset();
 }
 
-ULevelVariantSets* FDatasmithImporter::FinalizeLevelVariantSets(ULevelVariantSets* SourceLevelVariantSets, const TCHAR* VariantsFolderPath, ULevelVariantSets* ExistingLevelVariantSets)
+ULevelVariantSets* FDatasmithImporter::FinalizeLevelVariantSets( ULevelVariantSets* SourceLevelVariantSets, const TCHAR* VariantsFolderPath, ULevelVariantSets* ExistingLevelVariantSets )
 {
-	return Cast< ULevelVariantSets >(DatasmithImporterImpl::PublicizeAsset(SourceLevelVariantSets, VariantsFolderPath, ExistingLevelVariantSets));
+	return Cast< ULevelVariantSets >( DatasmithImporterImpl::PublicizeAsset( SourceLevelVariantSets, VariantsFolderPath, ExistingLevelVariantSets ) );
 }
 
 void FDatasmithImporter::ImportMetaDataForObject(FDatasmithImportContext& ImportContext, const TSharedRef<IDatasmithElement>& DatasmithElement, UObject* Object)
 {
-	if (!Object)
+	if ( !Object )
 	{
 		return;
 	}
@@ -1983,57 +2033,57 @@ void FDatasmithImporter::ImportMetaDataForObject(FDatasmithImportContext& Import
 	UDatasmithAssetUserData::FMetaDataContainer MetaData;
 
 	// Add Datasmith meta data
-	MetaData.Add(UDatasmithAssetUserData::UniqueIdMetaDataKey, DatasmithElement->GetName());
+	MetaData.Add( UDatasmithAssetUserData::UniqueIdMetaDataKey, DatasmithElement->GetName() );
 
 	// Check if there's metadata associated with the given element
-	const TSharedPtr<IDatasmithMetaDataElement>& MetaDataElement = ImportContext.Scene->GetMetaData(DatasmithElement);
-	if (MetaDataElement.IsValid())
+	const TSharedPtr<IDatasmithMetaDataElement>& MetaDataElement = ImportContext.Scene->GetMetaData( DatasmithElement );
+	if ( MetaDataElement.IsValid() )
 	{
 		int32 PropertiesCount = MetaDataElement->GetPropertiesCount();
-		MetaData.Reserve(PropertiesCount);
-		for (int32 PropertyIndex = 0; PropertyIndex < PropertiesCount; ++PropertyIndex)
+		MetaData.Reserve( PropertiesCount );
+		for ( int32 PropertyIndex = 0; PropertyIndex < PropertiesCount; ++PropertyIndex )
 		{
-			const TSharedPtr<IDatasmithKeyValueProperty>& Property = MetaDataElement->GetProperty(PropertyIndex);
-			MetaData.Add(Property->GetName(), Property->GetValue());
+			const TSharedPtr<IDatasmithKeyValueProperty>& Property = MetaDataElement->GetProperty( PropertyIndex );
+			MetaData.Add( Property->GetName(), Property->GetValue() );
 		}
 
 		MetaData.KeySort(FNameLexicalLess());
 	}
 
-	if (MetaData.Num() > 0)
+	if ( MetaData.Num() > 0 )
 	{
 		// For AActor, the interface is actually implemented by the ActorComponent
-		AActor* Actor = Cast<AActor>(Object);
-		if (Actor)
+		AActor* Actor = Cast<AActor>( Object );
+		if ( Actor )
 		{
 			UActorComponent* ActorComponent = Actor->GetRootComponent();
-			if (ActorComponent)
+			if ( ActorComponent )
 			{
 				Object = ActorComponent;
 			}
 		}
 
-		if (Object->GetClass()->ImplementsInterface(UInterface_AssetUserData::StaticClass()))
+		if ( Object->GetClass()->ImplementsInterface( UInterface_AssetUserData::StaticClass() ) )
 		{
-			IInterface_AssetUserData* AssetUserData = Cast< IInterface_AssetUserData >(Object);
+			IInterface_AssetUserData* AssetUserData = Cast< IInterface_AssetUserData >( Object );
 
 			UDatasmithAssetUserData* DatasmithUserData = AssetUserData->GetAssetUserData< UDatasmithAssetUserData >();
 
-			if (!DatasmithUserData)
+			if ( !DatasmithUserData )
 			{
-				DatasmithUserData = NewObject<UDatasmithAssetUserData>(Object, NAME_None, RF_Public | RF_Transactional);
-				AssetUserData->AddAssetUserData(DatasmithUserData);
+				DatasmithUserData = NewObject<UDatasmithAssetUserData>( Object, NAME_None, RF_Public | RF_Transactional );
+				AssetUserData->AddAssetUserData( DatasmithUserData );
 			}
 
-			DatasmithUserData->MetaData = MoveTemp(MetaData);
+			DatasmithUserData->MetaData = MoveTemp( MetaData );
 		}
 	}
 }
 
-void FDatasmithImporter::FilterElementsToImport(FDatasmithImportContext& ImportContext)
+void FDatasmithImporter::FilterElementsToImport( FDatasmithImportContext& ImportContext )
 {
 	// Initialize the filtered scene as a copy of the original scene. We will use it to then filter out items to import.
-	ImportContext.FilteredScene = FDatasmithSceneFactory::DuplicateScene(ImportContext.Scene.ToSharedRef());
+	ImportContext.FilteredScene = FDatasmithSceneFactory::DuplicateScene( ImportContext.Scene.ToSharedRef() );
 
 	// Filter meshes to import by consulting the AssetRegistry to see if that asset already exist
 	// or if it changed at all, if deemed identical filter the mesh out of the current import
@@ -2087,25 +2137,25 @@ void FDatasmithImporter::FilterElementsToImport(FDatasmithImportContext& ImportC
 	const TMap< FName, TSoftObjectPtr< UStaticMesh > >& StaticMeshes = ImportContext.SceneAsset->StaticMeshes;
 	for (int32 MeshIndex = 0; MeshIndex < ImportContext.Scene->GetMeshesCount(); ++MeshIndex)
 	{
-		TSharedRef< IDatasmithMeshElement > MeshElement = ImportContext.Scene->GetMesh(MeshIndex).ToSharedRef();
+		TSharedRef< IDatasmithMeshElement > MeshElement = ImportContext.Scene->GetMesh( MeshIndex ).ToSharedRef();
 
 		bool bNeedsReimport = true;
 		FString AssetName = MeshElement->GetName();
-		if (StaticMeshes.Contains(MeshElement->GetName()))
+		if ( StaticMeshes.Contains( MeshElement->GetName() ) )
 		{
-			AssetName = StaticMeshes[MeshElement->GetName()].ToString();
-			bNeedsReimport = ElementNeedsReimport(AssetName, MeshElement, ImportContext.Options->FilePath);
+			AssetName = StaticMeshes[ MeshElement->GetName() ].ToString();
+			bNeedsReimport = ElementNeedsReimport(AssetName, MeshElement, ImportContext.Options->FilePath );
 		}
 
-		if (bNeedsReimport)
+		if ( bNeedsReimport )
 		{
-			ImportContext.FilteredScene->AddMesh(MeshElement);
+			ImportContext.FilteredScene->AddMesh( MeshElement );
 		}
 		// If the mesh element does not need to be re-imported, register its name
 		else
 		{
-			const FAssetData AssetData = AssetRegistry.GetAssetByObjectPath(*AssetName);
-			ImportContext.AssetsContext.StaticMeshNameProvider.AddExistingName(FPaths::GetBaseFilename(AssetData.PackageName.ToString()));
+			const FAssetData AssetData = AssetRegistry.GetAssetByObjectPath( *AssetName);
+			ImportContext.AssetsContext.StaticMeshNameProvider.AddExistingName( FPaths::GetBaseFilename( AssetData.PackageName.ToString() ) );
 		}
 	}
 
@@ -2118,21 +2168,21 @@ void FDatasmithImporter::FilterElementsToImport(FDatasmithImportContext& ImportC
 
 		bool bNeedsReimport = true;
 		FString AssetName = TextureElement->GetName();
-		if (Textures.Contains(TextureElement->GetName()))
+		if ( Textures.Contains( TextureElement->GetName() ) )
 		{
 			AssetName = Textures[TextureElement->GetName()].ToString();
-			bNeedsReimport = ElementNeedsReimport(AssetName, TextureElement, ImportContext.Options->FilePath);
+			bNeedsReimport = ElementNeedsReimport(AssetName, TextureElement, ImportContext.Options->FilePath );
 		}
 
-		if (bNeedsReimport)
+		if ( bNeedsReimport )
 		{
-			ImportContext.FilteredScene->AddTexture(TextureElement);
+			ImportContext.FilteredScene->AddTexture( TextureElement );
 		}
 		// If the texture element does not need to be re-imported, register its name
 		else
 		{
-			const FAssetData AssetData = AssetRegistry.GetAssetByObjectPath(*AssetName);
-			ImportContext.AssetsContext.TextureNameProvider.AddExistingName(FPaths::GetBaseFilename(AssetData.PackageName.ToString()));
+			const FAssetData AssetData = AssetRegistry.GetAssetByObjectPath( *AssetName);
+			ImportContext.AssetsContext.TextureNameProvider.AddExistingName( FPaths::GetBaseFilename( AssetData.PackageName.ToString() ) );
 		}
 	}
 }
@@ -2142,11 +2192,11 @@ void FDatasmithImporter::FinalizeImport(FDatasmithImportContext& ImportContext,
 	TRACE_CPUPROFILER_EVENT_SCOPE(FDatasmithImporter::FinalizeImport);
 
 	const int32 NumImportedObjects = ImportContext.ImportedStaticMeshes.Num() +
-		ImportContext.ImportedTextures.Num() +
-		ImportContext.ImportedMaterialFunctions.Num() +
-		ImportContext.ImportedMaterials.Num() +
-		ImportContext.ImportedLevelSequences.Num() +
-		ImportContext.ImportedLevelVariantSets.Num();
+									 ImportContext.ImportedTextures.Num() +
+									 ImportContext.ImportedMaterialFunctions.Num() +
+									 ImportContext.ImportedMaterials.Num() +
+									 ImportContext.ImportedLevelSequences.Num() +
+									 ImportContext.ImportedLevelVariantSets.Num();
 	const int32 NumAssetsToFinalize = ValidAssets.Num() == 0 ? NumImportedObjects : ValidAssets.Num() + ImportContext.ImportedLevelSequences.Num() + ImportContext.ImportedLevelVariantSets.Num();
 	const int32 NumStaticMeshToBuild = ImportContext.ImportedStaticMeshes.Num();
 
@@ -2251,7 +2301,7 @@ void FDatasmithImporter::FinalizeImport(FDatasmithImportContext& ImportContext,
 		UMaterialInterface* ExistingMaterial = ExistingMaterialPtr.Get();
 
 		FString SourcePackagePath = SourceMaterialInterface->GetOutermost()->GetName();
-		FString DestinationPackagePath = SourcePackagePath.Replace(*TransientFolderPath, *RootFolderPath, ESearchCase::CaseSensitive);
+		FString DestinationPackagePath = SourcePackagePath.Replace( *TransientFolderPath, *RootFolderPath, ESearchCase::CaseSensitive );
 
 		if (UMaterial* SourceMaterial = Cast< UMaterial >(SourceMaterialInterface))
 		{
@@ -2311,7 +2361,7 @@ void FDatasmithImporter::FinalizeImport(FDatasmithImportContext& ImportContext,
 		UStaticMesh* ExistingStaticMesh = ExistingStaticMeshPtr.Get();
 
 		FString SourcePackagePath = SourceStaticMesh->GetOutermost()->GetName();
-		FString DestinationPackagePath = SourcePackagePath.Replace(*TransientFolderPath, *RootFolderPath, ESearchCase::CaseSensitive);
+		FString DestinationPackagePath = SourcePackagePath.Replace( *TransientFolderPath, *RootFolderPath, ESearchCase::CaseSensitive );
 
 		ExistingStaticMeshPtr = FinalizeStaticMesh(SourceStaticMesh, *DestinationPackagePath, ExistingStaticMesh, &ReferencesToRemap, false);
 		DatasmithImporterImpl::CheckAssetPersistenceValidity(ExistingStaticMeshPtr.Get(), ImportContext);
@@ -2350,7 +2400,7 @@ void FDatasmithImporter::FinalizeImport(FDatasmithImportContext& ImportContext,
 		ULevelSequence* ExistingLevelSequence = ExistingLevelSequencePtr.Get();
 
 		FString SourcePackagePath = SourceLevelSequence->GetOutermost()->GetName();
-		FString DestinationPackagePath = SourcePackagePath.Replace(*TransientFolderPath, *RootFolderPath, ESearchCase::CaseSensitive);
+		FString DestinationPackagePath = SourcePackagePath.Replace( *TransientFolderPath, *RootFolderPath, ESearchCase::CaseSensitive );
 
 		ExistingLevelSequencePtr = FinalizeLevelSequence(SourceLevelSequence, *DestinationPackagePath, ExistingLevelSequence);
 		ImportContext.SceneAsset->RegisterPreWorldRenameCallback();
@@ -2378,7 +2428,7 @@ void FDatasmithImporter::FinalizeImport(FDatasmithImportContext& ImportContext,
 		ULevelVariantSets* ExistingLevelVariantSets = ExistingLevelVariantSetsPtr.Get();
 
 		FString SourcePackagePath = SourceLevelVariantSets->GetOutermost()->GetName();
-		FString DestinationPackagePath = SourcePackagePath.Replace(*TransientFolderPath, *RootFolderPath, ESearchCase::CaseSensitive);
+		FString DestinationPackagePath = SourcePackagePath.Replace( *TransientFolderPath, *RootFolderPath, ESearchCase::CaseSensitive );
 
 		ExistingLevelVariantSetsPtr = FinalizeLevelVariantSets(SourceLevelVariantSets, *DestinationPackagePath, ExistingLevelVariantSets);
 		ImportContext.SceneAsset->RegisterPreWorldRenameCallback();
@@ -2389,16 +2439,16 @@ void FDatasmithImporter::FinalizeImport(FDatasmithImportContext& ImportContext,
 	{
 		TMap<FSoftObjectPath, FSoftObjectPath> AssetRedirectorMap;
 
-		for (const TPair< UObject*, UObject* >& ReferenceToRemap : ReferencesToRemap)
+		for ( const TPair< UObject*, UObject* >& ReferenceToRemap : ReferencesToRemap )
 		{
-			AssetRedirectorMap.Emplace(ReferenceToRemap.Key) = ReferenceToRemap.Value;
+			AssetRedirectorMap.Emplace( ReferenceToRemap.Key ) = ReferenceToRemap.Value;
 		}
 
 		IAssetTools& AssetTools = FModuleManager::LoadModuleChecked<FAssetToolsModule>("AssetTools").Get();
 		AssetTools.RenameReferencingSoftObjectPaths(PackagesToCheck, AssetRedirectorMap);
 	}
 
-	if (ImportContext.ShouldImportActors())
+	if ( ImportContext.ShouldImportActors() )
 	{
 		FinalizeActors(ImportContext, &ReferencesToRemap);
 	}
@@ -2417,4 +2467,4 @@ void FDatasmithImporter::FinalizeImport(FDatasmithImportContext& ImportContext,
 }
 
 
-#undef LOCTEXT_NAMESPACE
\ No newline at end of file
+#undef LOCTEXT_NAMESPACE
diff --git a/Engine/Plugins/Enterprise/DatasmithImporter/Source/DatasmithImporter/Private/DatasmithTextureImporter.cpp b/Engine/Plugins/Enterprise/DatasmithImporter/Source/DatasmithImporter/Private/DatasmithTextureImporter.cpp
--- a/Engine/Plugins/Enterprise/DatasmithImporter/Source/DatasmithImporter/Private/DatasmithTextureImporter.cpp
+++ b/Engine/Plugins/Enterprise/DatasmithImporter/Source/DatasmithImporter/Private/DatasmithTextureImporter.cpp
@@ -40,6 +40,7 @@ namespace
 {
 	static bool ResizeTexture(const TCHAR* Filename, const TCHAR* ResizedFilename, bool bCreateNormal, FDatasmithImportContext& ImportContext)
 	{
+		TRACE_CPUPROFILER_EVENT_SCOPE(ResizeTexture);
 
 		EDSTextureUtilsError ErrorCode = FDatasmithTextureResize::ResizeTexture(Filename, ResizedFilename, EDSResizeTextureMode::NearestPowerOfTwo, MaxTextureSize, bCreateNormal);
 
@@ -116,7 +117,7 @@ bool FDatasmithTextureImporter::ResizeTextureElement(const TSharedPtr<IDatasmith
 		Extension = TEXT(".exr");
 	}
 
-	ResizedFilename = FPaths::Combine(TempDir, FPaths::GetBaseFilename(Filename) + Extension);
+	ResizedFilename = FPaths::Combine(TempDir, LexToString(FGuid::NewGuid()) + Extension);
 
 	const bool bGenerateNormalMap = ( TextureElement->GetTextureMode() == EDatasmithTextureMode::Bump );
 
@@ -130,7 +131,62 @@ bool FDatasmithTextureImporter::ResizeTextureElement(const TSharedPtr<IDatasmith
 	return bResult;
 }
 
-UTexture* FDatasmithTextureImporter::CreateTexture(const TSharedPtr<IDatasmithTextureElement>& TextureElement)
+bool FDatasmithTextureImporter::GetTextureData(const TSharedPtr<IDatasmithTextureElement>& TextureElement, TArray<uint8>& TextureData, FString& Extension)
+{
+	const FString Filename = TextureElement->GetFile();
+	if (!Filename.IsEmpty())
+	{
+		// load from a file path
+
+		FString ImageFileName;
+		if (!ResizeTextureElement(TextureElement, ImageFileName))
+		{
+			return false;
+		}
+
+		// try opening from absolute path
+		if (!(FFileHelper::LoadFileToArray(TextureData, *ImageFileName) && TextureData.Num() > 0))
+		{
+			ImportContext.LogWarning(FText::Format(LOCTEXT("InvalidData", "Unable to find Texture file {0}."), FText::FromString(Filename)));
+			return false;
+		}
+
+		Extension         = FPaths::GetExtension(ImageFileName).ToLower();
+	}
+	else
+	{
+		// load from memory source
+		EDatasmithTextureFormat TextureFormat;
+		uint32                  TextureSize;
+
+		const uint8* PtrTextureData = TextureElement->GetData(TextureSize, TextureFormat);
+
+		if (PtrTextureData == nullptr || !TextureSize)
+		{
+			return false;
+		}
+
+		TextureData.SetNumUninitialized(TextureSize);
+		FPlatformMemory::Memcpy(TextureData.GetData(), PtrTextureData, TextureSize);
+
+		switch (TextureFormat)
+		{
+			case EDatasmithTextureFormat::PNG:
+				Extension = TEXT("png");
+				break;
+			case EDatasmithTextureFormat::JPEG:
+				Extension = TEXT("jpeg");
+				break;
+			default:
+				check(false);
+		}
+	}
+
+	return true;
+}
+
+
+UTexture* FDatasmithTextureImporter::CreateTexture(const TSharedPtr<IDatasmithTextureElement>& TextureElement, const TArray<uint8>& TextureData, const FString& Extension)
 {
 	const FString TextureLabel = TextureElement->GetLabel();
 	const FString TextureName = TextureLabel.Len() > 0 ? ImportContext.AssetsContext.TextureNameProvider.GenerateUniqueName(TextureLabel) : TextureElement->GetName();
@@ -173,58 +229,9 @@ UTexture* FDatasmithTextureImporter::CreateTexture(const TSharedPtr<IDatasmithTe
 
 	const float RGBCurve = TextureElement->GetRGBCurve();
 
-	FString       Extension;
-	const uint8*  PtrTextureData;
-	const uint8*  PtrTextureDataEnd;
-	TArray<uint8> FileTextureData;
+	const uint8* PtrTextureData    = TextureData.GetData();
+	const uint8* PtrTextureDataEnd = PtrTextureData + TextureData.Num();
 	const FString Filename = TextureElement->GetFile();
-	if (!Filename.IsEmpty())
-	{
-		// load from a file path
-
-		FString ImageFileName;
-		if (!ResizeTextureElement(TextureElement, ImageFileName))
-		{
-			return nullptr;
-		}
-
-		// try opening from absolute path
-		if (!(FFileHelper::LoadFileToArray(FileTextureData, *ImageFileName) && FileTextureData.Num() > 0))
-		{
-			ImportContext.LogWarning(FText::Format(LOCTEXT("NoTextureFileFound", "Unable to find Texture file {0}."), FText::FromString(Filename)));
-			return nullptr;
-		}
-
-		PtrTextureData    = FileTextureData.GetData();
-		PtrTextureDataEnd = PtrTextureData + FileTextureData.Num();
-		Extension         = FPaths::GetExtension(ImageFileName).ToLower();
-	}
-	else
-	{
-		// load from memory source
-
-		EDatasmithTextureFormat TextureFormat;
-		uint32                  TextureSize;
-
-		PtrTextureData    = TextureElement->GetData(TextureSize, TextureFormat);
-		PtrTextureDataEnd = PtrTextureData + TextureSize;
-
-		if (PtrTextureData == nullptr || !TextureSize)
-			return nullptr;
-
-		switch (TextureFormat)
-		{
-			case EDatasmithTextureFormat::PNG:
-				Extension = TEXT("png");
-				break;
-			case EDatasmithTextureFormat::JPEG:
-				Extension = TEXT("jpeg");
-				break;
-			default:
-				check(false);
-		}
-	}
-
 	UPackage* TextureOuter = ImportContext.AssetsContext.TexturesImportPackage.Get();
 
 	// This has to be called explicitly each time we create a texture since the flag gets reset in FactoryCreateBinary
diff --git a/Engine/Plugins/Enterprise/DatasmithImporter/Source/DatasmithImporter/Private/Utility/DatasmithTextureResize.cpp b/Engine/Plugins/Enterprise/DatasmithImporter/Source/DatasmithImporter/Private/Utility/DatasmithTextureResize.cpp
--- a/Engine/Plugins/Enterprise/DatasmithImporter/Source/DatasmithImporter/Private/Utility/DatasmithTextureResize.cpp
+++ b/Engine/Plugins/Enterprise/DatasmithImporter/Source/DatasmithImporter/Private/Utility/DatasmithTextureResize.cpp
@@ -107,7 +107,7 @@ class FFreeImageWrapper
 public:
 	static bool IsValid() { return FreeImageDllHandle != nullptr; }
 
-	static void FreeImage_Initialise(bool bLoadLocalPluginsOnly); // Loads and inits FreeImage on first call
+	static void FreeImage_Initialise(); // Loads and inits FreeImage on first call
 
 private:
 	static void* FreeImageDllHandle; // Lazy init on first use, never release for now
@@ -134,6 +134,11 @@ bool FDatasmithTextureResize::IsSupportedTextureExtension(const FString& Extensi
 			Extension == TEXT(".hdr");
 }
 
+void FDatasmithTextureResize::Initialize()
+{
+	FFreeImageWrapper::FreeImage_Initialise();
+}
+
 bool FDatasmithTextureResize::GetBestTextureExtension(const TCHAR* Source, FString& Extension)
 {
 	Extension = FPaths::GetExtension(Source, true).ToLower();
@@ -144,7 +149,7 @@ bool FDatasmithTextureResize::GetBestTextureExtension(const TCHAR* Source, FStri
 
 	Extension.Reset();
 
-	FFreeImageWrapper::FreeImage_Initialise(false);
+	Initialize();
 
 	//image format
 	FREE_IMAGE_FORMAT FileType = FIF_UNKNOWN;
@@ -212,7 +217,7 @@ EDSTextureUtilsError FDatasmithTextureResize::ResizeTexture(const TCHAR* Source,
 		return EDSTextureUtilsError::FileNotFound;
 	}
 
-	FFreeImageWrapper::FreeImage_Initialise(false);
+	Initialize();
 
 	if ( !FFreeImageWrapper::IsValid() )
 	{
@@ -656,13 +661,17 @@ void DatasmithTextureResizeInternal::GetBitmapPixelInfo(FIBITMAP* Bitmap, int32&
 
 void* FFreeImageWrapper::FreeImageDllHandle = nullptr;
 
-void FFreeImageWrapper::FreeImage_Initialise(bool bLoadLocalPluginsOnly)
+void FFreeImageWrapper::FreeImage_Initialise()
 {
 	if ( FreeImageDllHandle != nullptr )
 	{
 		return;
 	}
 
+	// Push/PopDllDirectory are soooooo not threadsafe!
+	// Must load library in main thread before doing parallel processing
+	check(IsInGameThread());
+
 	if ( FreeImageDllHandle == nullptr )
 	{
 		FString FreeImageDir = FPaths::Combine( FPaths::EngineDir(), TEXT("Binaries/ThirdParty/FreeImage"), FPlatformProcess::GetBinariesSubdirectory() );
@@ -674,7 +683,7 @@ void FFreeImageWrapper::FreeImage_Initialise(bool bLoadLocalPluginsOnly)
 
 	if ( FreeImageDllHandle )
 	{
-		::FreeImage_Initialise((BOOL)bLoadLocalPluginsOnly);
+		::FreeImage_Initialise();
 	}
 }
 
diff --git a/Engine/Plugins/Enterprise/DatasmithImporter/Source/DatasmithImporter/Public/DatasmithImportContext.h b/Engine/Plugins/Enterprise/DatasmithImporter/Source/DatasmithImporter/Public/DatasmithImportContext.h
--- a/Engine/Plugins/Enterprise/DatasmithImporter/Source/DatasmithImporter/Public/DatasmithImportContext.h
+++ b/Engine/Plugins/Enterprise/DatasmithImporter/Source/DatasmithImporter/Public/DatasmithImportContext.h
@@ -236,8 +236,8 @@ struct DATASMITHIMPORTER_API FDatasmithImportContext
 	/** List of previously parsed IES files */
 	TSet<FString> ParsedIesFiles;
 
-	/** Indicates if the user has cancelled the import process  */
-	bool bUserCancelled;
+	/** Indicates if the user has canceled the import process  */
+	TAtomic<bool> bUserCancelled;
 
 	bool bIsAReimport;
 	bool bImportedViaScript;
diff --git a/Engine/Plugins/Enterprise/DatasmithImporter/Source/DatasmithImporter/Public/DatasmithImporter.h b/Engine/Plugins/Enterprise/DatasmithImporter/Source/DatasmithImporter/Public/DatasmithImporter.h
--- a/Engine/Plugins/Enterprise/DatasmithImporter/Source/DatasmithImporter/Public/DatasmithImporter.h
+++ b/Engine/Plugins/Enterprise/DatasmithImporter/Source/DatasmithImporter/Public/DatasmithImporter.h
@@ -69,7 +69,7 @@ struct DATASMITHIMPORTER_API FDatasmithImporter
 	/**
 	 * @return The imported texture
 	 */
-	static UTexture* ImportTexture( FDatasmithImportContext& ImportContext, class FDatasmithTextureImporter& DatasmithTextureImporter, TSharedRef< IDatasmithTextureElement > TextureElement, UTexture* ExistingTexture );
+	static UTexture* ImportTexture( FDatasmithImportContext& ImportContext, class FDatasmithTextureImporter& DatasmithTextureImporter, TSharedRef< IDatasmithTextureElement > TextureElement, UTexture* ExistingTexture, const TArray<uint8>& TextureData, const FString& Extension );
 
 	/**
 	 * Copy a texture from its transient package into its final public package and builds it.
diff --git a/Engine/Plugins/Enterprise/DatasmithImporter/Source/DatasmithImporter/Public/DatasmithTextureImporter.h b/Engine/Plugins/Enterprise/DatasmithImporter/Source/DatasmithImporter/Public/DatasmithTextureImporter.h
--- a/Engine/Plugins/Enterprise/DatasmithImporter/Source/DatasmithImporter/Public/DatasmithTextureImporter.h
+++ b/Engine/Plugins/Enterprise/DatasmithImporter/Source/DatasmithImporter/Public/DatasmithTextureImporter.h
@@ -17,7 +17,8 @@ public:
 	FDatasmithTextureImporter(FDatasmithImportContext& InImportContext);
 	~FDatasmithTextureImporter();
 
-	UTexture* CreateTexture(const TSharedPtr<IDatasmithTextureElement>& TextureElement);
+	UTexture* CreateTexture(const TSharedPtr<IDatasmithTextureElement>& TextureElement, const TArray<uint8>& TextureData, const FString& Extension);
+	bool GetTextureData(const TSharedPtr<IDatasmithTextureElement>& TextureElement, TArray<uint8>& TextureData, FString& Extension);
 
 private:
 	bool ResizeTextureElement(const TSharedPtr<IDatasmithTextureElement>& TextureElement, FString& ResizedFilename);
diff --git a/Engine/Plugins/Enterprise/DatasmithImporter/Source/DatasmithImporter/Public/Utility/DatasmithTextureResize.h b/Engine/Plugins/Enterprise/DatasmithImporter/Source/DatasmithImporter/Public/Utility/DatasmithTextureResize.h
--- a/Engine/Plugins/Enterprise/DatasmithImporter/Source/DatasmithImporter/Public/Utility/DatasmithTextureResize.h
+++ b/Engine/Plugins/Enterprise/DatasmithImporter/Source/DatasmithImporter/Public/Utility/DatasmithTextureResize.h
@@ -10,6 +10,11 @@ class DATASMITHIMPORTER_API FDatasmithTextureResize
 {
 public:
 #if WITH_FREEIMAGE_LIB
+	/**
+	 * Must be called from the main thread before doing any parallel work.
+	 */
+	static void Initialize();
+
 	/**
 	 * Returns true if image extension is supported by Unreal Engine
 	 * IMPORTANT: Extension is expected with the starting dot
@@ -32,6 +37,12 @@ public:
 	static EDSTextureUtilsError ResizeTexture( const TCHAR* Source, const TCHAR* Destination, EDSResizeTextureMode Mode, uint32 MaxSize, bool bGenerateNormalMap );
 
 #else
+	/**
+	 * Must be called from the main thread before doing any parallel work.
+	 */
+	static void Initialize()
+	{ }
+
 	/**
 	* Returns true if image extension is supported by Unreal Engine
 	* IMPORTANT: Extension is expected with the starting dot
diff --git a/Engine/Source/Developer/Datasmith/DatasmithCore/Public/IDatasmithSceneElements.h b/Engine/Source/Developer/Datasmith/DatasmithCore/Public/IDatasmithSceneElements.h
--- a/Engine/Source/Developer/Datasmith/DatasmithCore/Public/IDatasmithSceneElements.h
+++ b/Engine/Source/Developer/Datasmith/DatasmithCore/Public/IDatasmithSceneElements.h
@@ -763,7 +763,7 @@ public:
 	 */
 	virtual void SetData(const uint8* InData, uint32 InDataSize, EDatasmithTextureFormat InFormat) = 0;
 
-	/** Retun the optional data, if loading from memory */
+	/** Return the optional data, if loading from memory. Must be callable from any thread. */
 	virtual const uint8* GetData(uint32& OutDataSize, EDatasmithTextureFormat& OutFormat) const = 0;
 
 	/** Return a MD5 hash of the content of the Texture Element. Used in CalculateElementHash to quickly identify Element with identical content */
