Use a map to search through imported static meshes instead of scanning
    - 10.9s -> 173 ms in FindAsset while importing DATASET-0006a

From dev-enterprise[REDACTED_HANDLE]

diff --git a/Engine/Plugins/Enterprise/DatasmithImporter/Source/DatasmithImporter/Private/DatasmithImporter.cpp b/Engine/Plugins/Enterprise/DatasmithImporter/Source/DatasmithImporter/Private/DatasmithImporter.cpp
--- a/Engine/Plugins/Enterprise/DatasmithImporter/Source/DatasmithImporter/Private/DatasmithImporter.cpp
+++ b/Engine/Plugins/Enterprise/DatasmithImporter/Source/DatasmithImporter/Private/DatasmithImporter.cpp
@@ -972,6 +972,8 @@ void FDatasmithImporter::ImportStaticMeshes( FDatasmithImportContext& ImportCont
 		{
 			ImportStaticMesh(ImportContext, MeshElement, ExistingStaticMesh, nullptr);
 		}
+
+		ImportContext.ImportedStaticMeshesByName.Add(MeshElement->GetName(), MeshElement);
 	}
 
 	//Just make sure there is no async task left running in case of a cancellation
@@ -1134,7 +1136,7 @@ void FDatasmithImporter::ImportTextures( FDatasmithImportContext& ImportContext
 
 		for ( int32 TextureIndex = 0; TextureIndex < FilteredTextureElements.Num(); TextureIndex++ )
 		{
-			if (ImportContext.bUserCancelled |= ImportContext.Warn->ReceivedUserCancel())
+			if ((ImportContext.bUserCancelled |= ImportContext.Warn->ReceivedUserCancel()))
 			{
 				// If operation has been canceled, just wait for other threads to also cancel
 				AsyncData[TextureIndex].Result.Wait();
diff --git a/Engine/Plugins/Enterprise/DatasmithImporter/Source/DatasmithImporter/Public/DatasmithImportContext.h b/Engine/Plugins/Enterprise/DatasmithImporter/Source/DatasmithImporter/Public/DatasmithImportContext.h
--- a/Engine/Plugins/Enterprise/DatasmithImporter/Source/DatasmithImporter/Public/DatasmithImportContext.h
+++ b/Engine/Plugins/Enterprise/DatasmithImporter/Source/DatasmithImporter/Public/DatasmithImportContext.h
@@ -251,6 +251,9 @@ struct DATASMITHIMPORTER_API FDatasmithImportContext
 	/** Map of imported mesh for each mesh element */
 	TMap< TSharedRef< IDatasmithMeshElement >, UStaticMesh* > ImportedStaticMeshes;
 
+	/** Register IDatasmithMeshElement by their name so they can be searched faster */
+	TMap< FString, TSharedRef < IDatasmithMeshElement > > ImportedStaticMeshesByName;
+
 	/** Map of imported texture for each texture element */
 	TMap< TSharedRef< IDatasmithTextureElement >, UTexture* > ImportedTextures;
 
diff --git a/Engine/Plugins/Enterprise/DatasmithImporter/Source/DatasmithImporter/Public/Utility/DatasmithImporterUtils.h b/Engine/Plugins/Enterprise/DatasmithImporter/Source/DatasmithImporter/Public/Utility/DatasmithImporterUtils.h
--- a/Engine/Plugins/Enterprise/DatasmithImporter/Source/DatasmithImporter/Public/Utility/DatasmithImporterUtils.h
+++ b/Engine/Plugins/Enterprise/DatasmithImporter/Source/DatasmithImporter/Public/Utility/DatasmithImporterUtils.h
@@ -156,6 +156,11 @@ struct FDatasmithFindAssetTypeHelper< UStaticMesh >
 	{
 		return SceneAsset ? &SceneAsset->StaticMeshes : nullptr;
 	}
+
+	static const TSharedRef<IDatasmithMeshElement>* GetImportedElementByName( const FDatasmithAssetsImportContext& AssetsContext, const TCHAR* ObjectPathName )
+	{
+		return AssetsContext.ParentContext.ImportedStaticMeshesByName.Find(ObjectPathName);
+	}
 };
 
 template<>
@@ -175,6 +180,11 @@ struct FDatasmithFindAssetTypeHelper< UTexture >
 	{
 		return SceneAsset ? &SceneAsset->Textures : nullptr;
 	}
+
+	static const TSharedRef<IDatasmithTextureElement>* GetImportedElementByName( const FDatasmithAssetsImportContext& AssetsContext, const TCHAR* ObjectPathName )
+	{
+		return nullptr;
+	}
 };
 
 template<>
@@ -194,6 +204,11 @@ struct FDatasmithFindAssetTypeHelper< UMaterialFunction >
 	{
 		return SceneAsset ? &SceneAsset->MaterialFunctions : nullptr;
 	}
+	
+	static const TSharedRef<IDatasmithBaseMaterialElement>* GetImportedElementByName( const FDatasmithAssetsImportContext& AssetsContext, const TCHAR* ObjectPathName )
+	{
+		return nullptr;
+	}
 };
 
 template<>
@@ -213,6 +228,11 @@ struct FDatasmithFindAssetTypeHelper< UMaterialInterface >
 	{
 		return SceneAsset ? &SceneAsset->Materials : nullptr;
 	}
+
+	static const TSharedRef<IDatasmithBaseMaterialElement>* GetImportedElementByName( const FDatasmithAssetsImportContext& AssetsContext, const TCHAR* ObjectPathName )
+	{
+		return nullptr;
+	}
 };
 
 template<>
@@ -263,13 +283,25 @@ inline ObjectType* FDatasmithImporterUtils::FindAsset( const FDatasmithAssetsImp
 
 	if ( FPaths::IsRelative( ObjectPathName ) )
 	{
+		const auto& ImportedElement   = FDatasmithFindAssetTypeHelper< ObjectType >::GetImportedElementByName( AssetsContext, ObjectPathName );
 		const auto& ImportedAssetsMap = FDatasmithFindAssetTypeHelper< ObjectType >::GetImportedAssetsMap( AssetsContext );
 
-		for ( const auto& ImportedAssetPair : ImportedAssetsMap )
+		if (ImportedElement)
+		{
+			auto Result = ImportedAssetsMap.Find(*ImportedElement);
+			if (Result)
+			{
+				return *Result;
+			}
+		}
+		else
 		{
-			if ( FCString::Stricmp( ImportedAssetPair.Key->GetName(), ObjectPathName ) == 0 )
+			for ( const auto& ImportedAssetPair : ImportedAssetsMap )
 			{
-				return ImportedAssetPair.Value;
+				if ( FCString::Stricmp( ImportedAssetPair.Key->GetName(), ObjectPathName ) == 0 )
+				{
+					return ImportedAssetPair.Value;
+				}
 			}
 		}
 
