Fix blend mode on nodes containing function call expressions

From dev-enterprise[REDACTED_HANDLE]

diff --git a/Engine/Plugins/Enterprise/DatasmithImporter/Source/DatasmithImporter/Private/DatasmithMaterialExpressions.cpp b/Engine/Plugins/Enterprise/DatasmithImporter/Source/DatasmithImporter/Private/DatasmithMaterialExpressions.cpp
--- a/Engine/Plugins/Enterprise/DatasmithImporter/Source/DatasmithImporter/Private/DatasmithMaterialExpressions.cpp
+++ b/Engine/Plugins/Enterprise/DatasmithImporter/Source/DatasmithImporter/Private/DatasmithMaterialExpressions.cpp
@@ -2476,7 +2476,7 @@ UMaterialExpressionMaterialFunctionCall* FDatasmithMaterialExpressions::BlendFun
 }
 
 // any kind of material
-UMaterialInterface* FDatasmithMaterialExpressions::CreateDatasmithMaterial(UPackage* Package, TSharedPtr< IDatasmithMaterialElement > MaterialElement, FDatasmithAssetsImportContext& AssetsContext,
+UMaterialInterface* FDatasmithMaterialExpressions::CreateDatasmithMaterial(UPackage* Package, const TSharedPtr< IDatasmithMaterialElement >& MaterialElement, FDatasmithAssetsImportContext& AssetsContext,
 																		  UMaterial* ExistingMaterial, EObjectFlags ObjectFlags)
 {
 	FString MaterialElementName = MaterialElement->GetName();
@@ -2799,7 +2799,7 @@ UMaterialExpression* FDatasmithMaterialExpressions::CreateExpression( IDatasmith
 	}
 }
 
-void FDatasmithMaterialExpressions::ConnectExpression( TSharedRef< IDatasmithUEPbrMaterialElement > MaterialElement, TArray< TStrongObjectPtr< UMaterialExpression > >& MaterialExpressions,
+void FDatasmithMaterialExpressions::ConnectExpression( const TSharedRef< IDatasmithUEPbrMaterialElement >& MaterialElement, TArray< TStrongObjectPtr< UMaterialExpression > >& MaterialExpressions,
 	IDatasmithMaterialExpression* MaterialExpression, FExpressionInput* MaterialInput, int32 OutputIndex )
 {
 	if ( !MaterialExpression || !MaterialInput )
@@ -2930,7 +2930,7 @@ UMaterialExpression* FDatasmithMaterialExpressions::CreateGenericExpression( IDa
 
 	for ( int32 PropertyIndex = 0; PropertyIndex < DatasmithGeneric.GetPropertiesCount(); ++PropertyIndex )
 	{
-		TSharedPtr< IDatasmithKeyValueProperty > KeyValueProperty = DatasmithGeneric.GetProperty( PropertyIndex );
+		const TSharedPtr< IDatasmithKeyValueProperty >& KeyValueProperty = DatasmithGeneric.GetProperty( PropertyIndex );
 
 		if ( !KeyValueProperty )
 		{
@@ -3016,7 +3016,7 @@ UMaterialExpression* FDatasmithMaterialExpressions::CreateFunctionCallExpression
 	return FunctionCall;
 }
 
-void FDatasmithMaterialExpressions::ConnectAnyExpression( TSharedRef< IDatasmithUEPbrMaterialElement > MaterialElement, TArray< TStrongObjectPtr< UMaterialExpression > >& MaterialExpressions,
+void FDatasmithMaterialExpressions::ConnectAnyExpression( const TSharedRef< IDatasmithUEPbrMaterialElement >& MaterialElement, TArray< TStrongObjectPtr< UMaterialExpression > >& MaterialExpressions,
 	IDatasmithMaterialExpression& DatasmithExpression, FExpressionInput* ExpressionInput, int32 OutputIndex )
 {
 	int32 ExpressionIndex = MaterialElement->GetExpressionIndex( &DatasmithExpression );
@@ -3037,7 +3037,7 @@ void FDatasmithMaterialExpressions::ConnectAnyExpression( TSharedRef< IDatasmith
 	}
 }
 
-void FDatasmithMaterialExpressions::CreateUEPbrMaterialGraph(TSharedPtr< IDatasmithUEPbrMaterialElement > MaterialElement, FDatasmithAssetsImportContext& AssetsContext, UObject* UnrealMaterialOrFunction)
+void FDatasmithMaterialExpressions::CreateUEPbrMaterialGraph(const TSharedPtr< IDatasmithUEPbrMaterialElement >& MaterialElement, FDatasmithAssetsImportContext& AssetsContext, UObject* UnrealMaterialOrFunction)
 {
 	TArray< TStrongObjectPtr< UMaterialExpression > > MaterialExpressions;
 
@@ -3073,7 +3073,7 @@ void FDatasmithMaterialExpressions::CreateUEPbrMaterialGraph(TSharedPtr< IDatasm
 
 	if (MaterialElement->GetOpacity().GetExpression())
 	{
-		if (GetUEPbrImportBlendMode(MaterialElement) == BLEND_Translucent)
+		if (GetUEPbrImportBlendMode(MaterialElement, AssetsContext) == BLEND_Translucent)
 		{
 			ConnectExpression( MaterialElement.ToSharedRef(), MaterialExpressions, MaterialElement->GetOpacity().GetExpression(), GetMaterialOrFunctionSlot( UnrealMaterialOrFunction, EDatasmithTextureSlot::OPACITY ), MaterialElement->GetOpacity().GetOutputIndex() );
 			
@@ -3089,7 +3089,7 @@ void FDatasmithMaterialExpressions::CreateUEPbrMaterialGraph(TSharedPtr< IDatasm
 	ConnectExpression( MaterialElement.ToSharedRef(), MaterialExpressions, MaterialElement->GetWorldDisplacement().GetExpression(), GetMaterialOrFunctionSlot( UnrealMaterialOrFunction, EDatasmithTextureSlot::DISPLACE ), MaterialElement->GetWorldDisplacement().GetOutputIndex() );
 }
 
-UMaterialFunction* FDatasmithMaterialExpressions::CreateUEPbrMaterialFunction(UPackage* Package, TSharedPtr< IDatasmithUEPbrMaterialElement > MaterialElement, FDatasmithAssetsImportContext& AssetsContext, UMaterial* ExistingMaterial, EObjectFlags ObjectFlags)
+UMaterialFunction* FDatasmithMaterialExpressions::CreateUEPbrMaterialFunction(UPackage* Package, const TSharedPtr< IDatasmithUEPbrMaterialElement >& MaterialElement, FDatasmithAssetsImportContext& AssetsContext, UMaterial* ExistingMaterial, EObjectFlags ObjectFlags)
 {
 	FString MaterialFunctionName = GenerateUniqueMaterialName(MaterialElement->GetParentLabel(), AssetsContext.MaterialFunctionNameProvider);
 
@@ -3109,20 +3109,41 @@ UMaterialFunction* FDatasmithMaterialExpressions::CreateUEPbrMaterialFunction(UP
 	return UnrealMaterialFunc;
 }
 
-EBlendMode FDatasmithMaterialExpressions::GetUEPbrImportBlendMode(TSharedPtr<IDatasmithUEPbrMaterialElement> MaterialElement)
+EBlendMode FDatasmithMaterialExpressions::GetUEPbrImportBlendMode(const TSharedPtr< IDatasmithUEPbrMaterialElement >& MaterialElement, const FDatasmithAssetsImportContext& AssetsContext)
 {
 	EBlendMode BlendMode = static_cast<EBlendMode>(MaterialElement->GetBlendMode());
-	
+
 	if ( MaterialElement->GetOpacity().GetExpression() &&
 		BlendMode != BLEND_Translucent && BlendMode != BLEND_Masked )
 	{
-		BlendMode = EBlendMode::BLEND_Translucent; // force translucent
+		return EBlendMode::BLEND_Translucent; // force translucent
+	}
+
+	// If this material is dependent on function call expression, we need to check if those are translucent to apply
+	// translucency attribute to this material too.
+	for (int32 ExpressionIndex = 0; ExpressionIndex < MaterialElement->GetExpressionsCount(); ExpressionIndex++)
+	{
+		IDatasmithMaterialExpression* MaterialExpression = MaterialElement->GetExpression(ExpressionIndex);
+		if (MaterialExpression && MaterialExpression->IsA(EDatasmithMaterialExpressionType::FunctionCall))
+		{
+			IDatasmithMaterialExpressionFunctionCall* MaterialExpressionFunctionCall = static_cast< IDatasmithMaterialExpressionFunctionCall* >(MaterialExpression);
+				
+			const TSharedRef<IDatasmithBaseMaterialElement>* MaterialFunction = FDatasmithFindAssetTypeHelper<UMaterialFunction>::GetImportedElementByName(AssetsContext, MaterialExpressionFunctionCall->GetFunctionPathName());
+			if (MaterialFunction && (*MaterialFunction)->IsA(EDatasmithElementType::UEPbrMaterial))
+			{
+				const TSharedRef<IDatasmithUEPbrMaterialElement> MaterialFunctionElement = StaticCastSharedRef<IDatasmithUEPbrMaterialElement>(*MaterialFunction);
+				if (GetUEPbrImportBlendMode(MaterialFunctionElement, AssetsContext) == EBlendMode::BLEND_Translucent)
+				{
+					return EBlendMode::BLEND_Translucent;
+				}
+			}
+		}
 	}
 
 	return BlendMode;
 }
 
-UMaterialInterface* FDatasmithMaterialExpressions::CreateUEPbrMaterial(UPackage* Package, TSharedPtr< IDatasmithUEPbrMaterialElement > MaterialElement, FDatasmithAssetsImportContext& AssetsContext,
+UMaterialInterface* FDatasmithMaterialExpressions::CreateUEPbrMaterial(UPackage* Package, const TSharedPtr< IDatasmithUEPbrMaterialElement >& MaterialElement, FDatasmithAssetsImportContext& AssetsContext,
 	UMaterial* ExistingMaterial, EObjectFlags ObjectFlags)
 {
 	FString MaterialName = GenerateUniqueMaterialName(MaterialElement->GetParentLabel(), AssetsContext.MasterMaterialNameProvider);
@@ -3148,7 +3169,7 @@ UMaterialInterface* FDatasmithMaterialExpressions::CreateUEPbrMaterial(UPackage*
 
 	UnrealMaterial->bUseMaterialAttributes = MaterialElement->GetUseMaterialAttributes();
 	UnrealMaterial->TwoSided = MaterialElement->GetTwoSided();
-	UnrealMaterial->BlendMode = GetUEPbrImportBlendMode(MaterialElement);
+	UnrealMaterial->BlendMode = GetUEPbrImportBlendMode(MaterialElement, AssetsContext);
 	
 	CreateUEPbrMaterialGraph(MaterialElement, AssetsContext, UnrealMaterial);
 
@@ -3168,7 +3189,7 @@ UMaterialInterface* FDatasmithMaterialExpressions::CreateUEPbrMaterial(UPackage*
 	return UnrealMaterial;
 }
 
-UMaterialInterface* FDatasmithMaterialExpressions::CreateUEPbrMaterialInstance(UPackage* Package, TSharedPtr< IDatasmithUEPbrMaterialElement > MaterialElement, FDatasmithAssetsImportContext& AssetsContext,
+UMaterialInterface* FDatasmithMaterialExpressions::CreateUEPbrMaterialInstance(UPackage* Package, const TSharedPtr< IDatasmithUEPbrMaterialElement >& MaterialElement, FDatasmithAssetsImportContext& AssetsContext,
 	UMaterialInterface* ParentMaterial, EObjectFlags ObjectFlags)
 {
 	FString MaterialName = GenerateUniqueMaterialName(MaterialElement->GetLabel(), AssetsContext.MaterialNameProvider);
@@ -3270,7 +3291,7 @@ UMaterialInterface* FDatasmithMaterialExpressions::CreateUEPbrMaterialInstance(U
 			{
 				FName ParameterName = GenerateParamName( *DatasmithExpression, TEXT("Generic"), ExpressionIndex + 1 );
 
-				TSharedPtr< IDatasmithKeyValueProperty > KeyValueProperty = GenericExpression->GetProperty( 0 );
+				const TSharedPtr< IDatasmithKeyValueProperty >& KeyValueProperty = GenericExpression->GetProperty( 0 );
 
 				if ( KeyValueProperty )
 				{
diff --git a/Engine/Plugins/Enterprise/DatasmithImporter/Source/DatasmithImporter/Private/DatasmithMaterialImporter.cpp b/Engine/Plugins/Enterprise/DatasmithImporter/Source/DatasmithImporter/Private/DatasmithMaterialImporter.cpp
--- a/Engine/Plugins/Enterprise/DatasmithImporter/Source/DatasmithImporter/Private/DatasmithMaterialImporter.cpp
+++ b/Engine/Plugins/Enterprise/DatasmithImporter/Source/DatasmithImporter/Private/DatasmithMaterialImporter.cpp
@@ -168,6 +168,7 @@ UMaterialFunction* FDatasmithMaterialImporter::CreateMaterialFunction( FDatasmit
 	if ( MaterialFunction != nullptr )
 	{
 		ImportContext.ImportedMaterialFunctions.Add( BaseMaterialElement ) = MaterialFunction;
+		ImportContext.ImportedMaterialFunctionsByName.Add( BaseMaterialElement->GetName(), BaseMaterialElement );
 	}
 
 	return MaterialFunction;
diff --git a/Engine/Plugins/Enterprise/DatasmithImporter/Source/DatasmithImporter/Public/DatasmithImportContext.h b/Engine/Plugins/Enterprise/DatasmithImporter/Source/DatasmithImporter/Public/DatasmithImportContext.h
--- a/Engine/Plugins/Enterprise/DatasmithImporter/Source/DatasmithImporter/Public/DatasmithImportContext.h
+++ b/Engine/Plugins/Enterprise/DatasmithImporter/Source/DatasmithImporter/Public/DatasmithImportContext.h
@@ -263,6 +263,9 @@ struct DATASMITHIMPORTER_API FDatasmithImportContext
 	/** Map of imported material function for each material element, they are only imported as a per-required basis */
 	TMap< TSharedRef< IDatasmithBaseMaterialElement >, UMaterialFunction* > ImportedMaterialFunctions;
 
+	/** Register IDatasmithMeshElement by their name so they can be searched faster */
+	TMap< FString, TSharedRef < IDatasmithBaseMaterialElement > > ImportedMaterialFunctionsByName;
+
 	/** List of potential parent materials with their hash. Used to create material instances from. */
 	TMap< int32, UMaterialInterface* > ImportedParentMaterials;
 
diff --git a/Engine/Plugins/Enterprise/DatasmithImporter/Source/DatasmithImporter/Public/DatasmithMaterialExpressions.h b/Engine/Plugins/Enterprise/DatasmithImporter/Source/DatasmithImporter/Public/DatasmithMaterialExpressions.h
--- a/Engine/Plugins/Enterprise/DatasmithImporter/Source/DatasmithImporter/Public/DatasmithMaterialExpressions.h
+++ b/Engine/Plugins/Enterprise/DatasmithImporter/Source/DatasmithImporter/Public/DatasmithMaterialExpressions.h
@@ -62,10 +62,10 @@ enum class EDatasmithTextureSlot
 class DATASMITHIMPORTER_API FDatasmithMaterialExpressions
 {
 public:
-	static UMaterialInterface* CreateDatasmithMaterial(UPackage* Package, TSharedPtr< IDatasmithMaterialElement > MaterialElement, FDatasmithAssetsImportContext& AssetsContext, UMaterial* ExistingMaterial, EObjectFlags ObjectFlags);
-	static UMaterialFunction* CreateUEPbrMaterialFunction(UPackage* Package, TSharedPtr< IDatasmithUEPbrMaterialElement > MaterialElement, FDatasmithAssetsImportContext& AssetsContext, UMaterial* ExistingMaterial, EObjectFlags ObjectFlags);
-	static UMaterialInterface* CreateUEPbrMaterial(UPackage* Package, TSharedPtr< IDatasmithUEPbrMaterialElement > MaterialElement, FDatasmithAssetsImportContext& AssetsContext, UMaterial* ExistingMaterial, EObjectFlags ObjectFlags);
-	static UMaterialInterface* CreateUEPbrMaterialInstance(UPackage* Package, TSharedPtr< IDatasmithUEPbrMaterialElement > MaterialElement, FDatasmithAssetsImportContext& AssetsContext, UMaterialInterface* ParentMaterial, EObjectFlags ObjectFlags);
+	static UMaterialInterface* CreateDatasmithMaterial(UPackage* Package, const TSharedPtr< IDatasmithMaterialElement >& MaterialElement, FDatasmithAssetsImportContext& AssetsContext, UMaterial* ExistingMaterial, EObjectFlags ObjectFlags);
+	static UMaterialFunction* CreateUEPbrMaterialFunction(UPackage* Package, const TSharedPtr< IDatasmithUEPbrMaterialElement >& MaterialElement, FDatasmithAssetsImportContext& AssetsContext, UMaterial* ExistingMaterial, EObjectFlags ObjectFlags);
+	static UMaterialInterface* CreateUEPbrMaterial(UPackage* Package, const TSharedPtr< IDatasmithUEPbrMaterialElement >& MaterialElement, FDatasmithAssetsImportContext& AssetsContext, UMaterial* ExistingMaterial, EObjectFlags ObjectFlags);
+	static UMaterialInterface* CreateUEPbrMaterialInstance(UPackage* Package, const TSharedPtr< IDatasmithUEPbrMaterialElement >& MaterialElement, FDatasmithAssetsImportContext& AssetsContext, UMaterialInterface* ParentMaterial, EObjectFlags ObjectFlags);
 	static UTextureLightProfile* CreateDatasmithIES(const FString& Filename, UPackage* Package, EObjectFlags ObjectFlags);
 	static UMaterialInterface* CreateDatasmithEnvironmentMaterial(UPackage* Package, const TSharedPtr< IDatasmithShaderElement >& ShaderElement, FDatasmithAssetsImportContext& AssetsContext,
 																  UMaterial* ExistingMaterial);
@@ -110,8 +110,8 @@ private:
 
 	static UMaterialExpression* AddCroppedUVMappingExpression(UObject* UnrealMaterial, const FDatasmithTextureSampler& UV, UMaterialExpressionTextureCoordinate* ToBeConnected);
 
-	static void CreateUEPbrMaterialGraph(TSharedPtr< IDatasmithUEPbrMaterialElement > MaterialElement, FDatasmithAssetsImportContext& AssetsContext, UObject* UnrealMaterialOrFunction);
-	static EBlendMode GetUEPbrImportBlendMode(TSharedPtr< IDatasmithUEPbrMaterialElement > MaterialElement);
+	static void CreateUEPbrMaterialGraph(const TSharedPtr< IDatasmithUEPbrMaterialElement >& MaterialElement, FDatasmithAssetsImportContext& AssetsContext, UObject* UnrealMaterialOrFunction);
+	static EBlendMode GetUEPbrImportBlendMode(const TSharedPtr< IDatasmithUEPbrMaterialElement >& MaterialElement, const FDatasmithAssetsImportContext& AssetsContext);
 	static EMaterialProperty DatasmithTextureSlotToMaterialProperty(EDatasmithTextureSlot InSlot);
 	static FExpressionInput* GetMaterialOrFunctionSlot(UObject* UnrealMatOrFunc, EDatasmithTextureSlot InSlot);
 	static void ConnectToSlot(UMaterialExpression* ToBeConnected, UObject* UnrealMatOrFunc, EDatasmithTextureSlot Slot, int32 InputChannel = 0);
@@ -158,6 +158,6 @@ private:
 	static UMaterialExpression* CreateGenericExpression( class IDatasmithMaterialExpressionGeneric& DatasmithGeneric, const FDatasmithAssetsImportContext& AssetsContext, UObject* UnrealMaterialOrFunction);
 	static UMaterialExpression* CreateFunctionCallExpression( class IDatasmithMaterialExpressionFunctionCall& DatasmithFunctionCall, const FDatasmithAssetsImportContext& AssetsContext, UObject* UnrealMaterialOrFunction);
 
-	static void ConnectExpression( TSharedRef< IDatasmithUEPbrMaterialElement > MaterialElement, TArray< TStrongObjectPtr< UMaterialExpression > >& MaterialExpressions, class IDatasmithMaterialExpression* MaterialExpression, FExpressionInput* MaterialInput, int32 OutputIndex);
-	static void ConnectAnyExpression( TSharedRef< IDatasmithUEPbrMaterialElement > MaterialElement, TArray< TStrongObjectPtr< UMaterialExpression > >& MaterialExpressions, class IDatasmithMaterialExpression& DatasmithExpression, FExpressionInput* ExpressionInput, int32 OutputIndex);
+	static void ConnectExpression( const TSharedRef< IDatasmithUEPbrMaterialElement >& MaterialElement, TArray< TStrongObjectPtr< UMaterialExpression > >& MaterialExpressions, class IDatasmithMaterialExpression* MaterialExpression, FExpressionInput* MaterialInput, int32 OutputIndex);
+	static void ConnectAnyExpression( const TSharedRef< IDatasmithUEPbrMaterialElement >& MaterialElement, TArray< TStrongObjectPtr< UMaterialExpression > >& MaterialExpressions, class IDatasmithMaterialExpression& DatasmithExpression, FExpressionInput* ExpressionInput, int32 OutputIndex);
 };
diff --git a/Engine/Plugins/Enterprise/DatasmithImporter/Source/DatasmithImporter/Public/Utility/DatasmithImporterUtils.h b/Engine/Plugins/Enterprise/DatasmithImporter/Source/DatasmithImporter/Public/Utility/DatasmithImporterUtils.h
--- a/Engine/Plugins/Enterprise/DatasmithImporter/Source/DatasmithImporter/Public/Utility/DatasmithImporterUtils.h
+++ b/Engine/Plugins/Enterprise/DatasmithImporter/Source/DatasmithImporter/Public/Utility/DatasmithImporterUtils.h
@@ -207,7 +207,7 @@ struct FDatasmithFindAssetTypeHelper< UMaterialFunction >
 	
 	static const TSharedRef<IDatasmithBaseMaterialElement>* GetImportedElementByName( const FDatasmithAssetsImportContext& AssetsContext, const TCHAR* ObjectPathName )
 	{
-		return nullptr;
+		return AssetsContext.ParentContext.ImportedMaterialFunctionsByName.Find(ObjectPathName);
 	}
 };
 
