Fixed issue with MaterialFunctions

diff --git a/Engine/Plugins/Enterprise/DatasmithImporter/Source/DatasmithImporter/Private/DatasmithConsumer.cpp b/Engine/Plugins/Enterprise/DatasmithImporter/Source/DatasmithImporter/Private/DatasmithConsumer.cpp
--- a/Engine/Plugins/Enterprise/DatasmithImporter/Source/DatasmithImporter/Private/DatasmithConsumer.cpp
+++ b/Engine/Plugins/Enterprise/DatasmithImporter/Source/DatasmithImporter/Private/DatasmithConsumer.cpp
@@ -639,6 +639,7 @@ namespace DatasmithConsumerUtils
 		// First skip UMaterial objects which are not referenced by a UmaterialInstance one
 		int32 MaterialCount = 0;
 		TSet< UMaterialInterface* > ParentMaterials;
+		TSet< UMaterialFunctionInterface* > MaterialFunctions;
 		for(TWeakObjectPtr<UObject>& AssetPtr : Assets)
 		{
 			if( UObject* Asset = AssetPtr.Get() )
@@ -655,9 +656,15 @@ namespace DatasmithConsumerUtils
 				}
 				else if(UMaterialInstance* MaterialInstance = Cast<UMaterialInstance>(Asset))
 				{
-					TSharedRef< IDatasmithMaterialElement > MaterialElement = FDatasmithSceneFactory::CreateMaterial( *AssetTag );
+					TSharedRef< IDatasmithBaseMaterialElement > MaterialElement = FDatasmithSceneFactory::CreateMaterial( *AssetTag );
 					MaterialElement->SetLabel( *MaterialInstance->GetName() );
 
+					if (UMaterial* SourceMaterial = Cast< UMaterial >(MaterialInstance))
+					{
+						MaterialElement = StaticCastSharedRef< IDatasmithBaseMaterialElement >( FDatasmithSceneFactory::CreateUEPbrMaterial( *AssetTag ) );
+						MaterialElement->SetLabel( *MaterialInstance->GetName() );
+					}
+
 					if ( UMaterialInterface* MaterialParent = MaterialInstance->Parent )
 					{
 						FString MaterialInstancePath = MaterialInstance->GetOutermost()->GetName();
@@ -734,7 +741,8 @@ namespace DatasmithConsumerUtils
 		// Second take care UMaterial objects which are not referenced by a UmaterialInstance one
 		for( TWeakObjectPtr<UObject>& AssetPtr : Assets )
 		{
-			if( UMaterial* Material = Cast<UMaterial>( AssetPtr.Get() ) )
+			UObject* AssetObject = AssetPtr.Get();
+			if( UMaterial* Material = Cast<UMaterial>( AssetObject ) )
 			{
 				if( !ParentMaterials.Contains( Material ) )
 				{
@@ -746,6 +754,25 @@ namespace DatasmithConsumerUtils
 					ImportContext.Scene->AddMaterial( MaterialElement );
 				}
 			}
+			else if( UMaterialFunction* MaterialFunction = Cast<UMaterialFunction>( AssetObject ) )
+			{
+				if( !MaterialFunctions.Contains( Cast<UMaterialFunctionInterface>( MaterialFunction ) ) )
+				{
+					FString AssetTag = FDatasmithImporterUtils::GetDatasmithElementIdString( MaterialFunction );
+
+					TSharedRef< IDatasmithUEPbrMaterialElement > UEPbrMaterialFunctionElement = FDatasmithSceneFactory::CreateUEPbrMaterial( *AssetTag );
+
+					UEPbrMaterialFunctionElement->SetLabel( *MaterialFunction->GetName() );
+					UEPbrMaterialFunctionElement->SetMaterialFunctionOnly( true );
+
+					TSharedRef< IDatasmithBaseMaterialElement > BaseMaterialElement = StaticCastSharedRef< IDatasmithBaseMaterialElement >( UEPbrMaterialFunctionElement );
+
+					ImportContext.ImportedMaterialFunctions.Add( BaseMaterialElement, MaterialFunction );
+					ImportContext.ImportedMaterialFunctionsByName.Add( BaseMaterialElement->GetName(), BaseMaterialElement );
+
+					ImportContext.Scene->AddMaterial( BaseMaterialElement );
+				}
+			}
 		}
 	}
 }
diff --git a/Engine/Plugins/Enterprise/DatasmithImporter/Source/DatasmithImporter/Private/DatasmithFileProducer.cpp b/Engine/Plugins/Enterprise/DatasmithImporter/Source/DatasmithImporter/Private/DatasmithFileProducer.cpp
--- a/Engine/Plugins/Enterprise/DatasmithImporter/Source/DatasmithImporter/Private/DatasmithFileProducer.cpp
+++ b/Engine/Plugins/Enterprise/DatasmithImporter/Source/DatasmithImporter/Private/DatasmithFileProducer.cpp
@@ -297,6 +297,19 @@ void UDatasmithFileProducer::SceneElementToWorld()
 	for ( TPair< TSharedRef< IDatasmithBaseMaterialElement >, UMaterialInterface* >& AssetPair : ImportContextPtr->ImportedMaterials )
 	{
 		Assets.Emplace( AssetPair.Value );
+
+		if (UMaterial* SourceMaterial = Cast< UMaterial >(AssetPair.Value))
+		{
+			SourceMaterial->RebuildExpressionTextureReferences();
+
+			for (FMaterialFunctionInfo& MaterialFunctionInfo : SourceMaterial->MaterialFunctionInfos)
+			{
+				if (MaterialFunctionInfo.Function && MaterialFunctionInfo.Function->GetOutermost() == SourceMaterial->GetOutermost())
+				{
+					Assets.Emplace( MaterialFunctionInfo.Function );
+				}
+			}
+		}
 	}
 
 	for ( TPair< int32, UMaterialInterface* >& AssetPair : ImportContextPtr->ImportedParentMaterials )
@@ -304,6 +317,11 @@ void UDatasmithFileProducer::SceneElementToWorld()
 		Assets.Emplace( AssetPair.Value );
 	}
 
+	for ( TPair< TSharedRef< IDatasmithBaseMaterialElement >, UMaterialFunction* >& AssetPair : ImportContextPtr->ImportedMaterialFunctions )
+	{
+		Assets.Emplace( AssetPair.Value );
+	}
+
 	for ( TPair< TSharedRef< IDatasmithLevelSequenceElement >, ULevelSequence* >& AssetPair : ImportContextPtr->ImportedLevelSequences )
 	{
 		Assets.Emplace( AssetPair.Value );
@@ -386,6 +404,7 @@ void UDatasmithFileProducer::PreventNameCollision()
 
 		// First pass: No UMaterial objects, parent materials are collected if applicable
 		TSet<UMaterialInterface*> ParentMaterials;
+		TSet<UMaterialFunctionInterface*> MaterialFunctions;
 
 		for(int32 Index = 0; Index < Assets.Num(); ++Index)
 		{
@@ -405,8 +424,23 @@ void UDatasmithFileProducer::PreventNameCollision()
 				{
 					MoveAsset( Object, TexturesImportPackage, false );
 				}
+				else if( Cast<UMaterialFunctionInterface>(Object) != nullptr )
+				{
+					MoveAsset( Object, MaterialsImportPackage, true );
+				}
 				else if( UMaterialInstance* MaterialInstance = Cast<UMaterialInstance>(Object) )
 				{
+					if (UMaterial* SourceMaterial = Cast< UMaterial >(MaterialInstance))
+					{
+						for (FMaterialFunctionInfo& MaterialFunctionInfo : SourceMaterial->MaterialFunctionInfos)
+						{
+							if (MaterialFunctionInfo.Function && MaterialFunctionInfo.Function->GetOutermost() == SourceMaterial->GetOutermost())
+							{
+								MaterialFunctions.Add( MaterialFunctionInfo.Function );
+							}
+						}
+					}
+
 					if ( UMaterialInterface* MaterialParent = MaterialInstance->Parent )
 					{
 						FString MaterialInstancePath = MaterialInstance->GetOutermost()->GetName();
@@ -450,6 +484,14 @@ void UDatasmithFileProducer::PreventNameCollision()
 					MoveAsset( Material, MaterialsImportPackage, true );
 				}
 			}
+			else if( UMaterialFunctionInterface* MaterialFunction = Cast<UMaterialFunctionInterface>( Assets[Index].Get() ) )
+			{
+				if( !MaterialFunctions.Contains( MaterialFunction ) )
+				{
+					PathsToDelete.Add( FPaths::GetPath( MaterialFunction->GetOutermost()->GetName() ) );
+					MoveAsset( MaterialFunction, MaterialsImportPackage, true );
+				}
+			}
 		}
 
 		// Apply soft object path redirection to identified packages
