Remove garbage collect during finalization in datasmith import
  - There was too much impact on performance and not enough memory savings.
  - 4m18s [at] 11.90GB -> 3m30s [at] 12.75GB for DATASET-0006a

From dev-enterprise[REDACTED_HANDLE]

diff --git a/Engine/Plugins/Enterprise/DatasmithImporter/Source/DatasmithImporter/Private/DatasmithImporter.cpp b/Engine/Plugins/Enterprise/DatasmithImporter/Source/DatasmithImporter/Private/DatasmithImporter.cpp
--- a/Engine/Plugins/Enterprise/DatasmithImporter/Source/DatasmithImporter/Private/DatasmithImporter.cpp
+++ b/Engine/Plugins/Enterprise/DatasmithImporter/Source/DatasmithImporter/Private/DatasmithImporter.cpp
@@ -75,24 +75,10 @@ extern UNREALED_API UEditorEngine* GEditor;
 
 namespace DatasmithImporterImpl
 {
-	static void SafetyGarbageCollect()
-	{
-		const int SkipCountThreshold = 1000; // Don't try to cleanup to frequently
-
-		static int SkippedCount = 0;
-		if (++SkippedCount > SkipCountThreshold)
-		{
-			CollectGarbage(GARBAGE_COLLECTION_KEEPFLAGS);
-			SkippedCount = 0;
-		}
-	}
-
 	static UObject* PublicizeAsset( UObject* SourceAsset, const TCHAR* DestinationPath, UObject* ExistingAsset )
 	{
 		TRACE_CPUPROFILER_EVENT_SCOPE(DatasmithImporterImpl::PublicizeAsset);
 
-		SafetyGarbageCollect();
-
 		UPackage* DestinationPackage;
 
 		if ( !ExistingAsset )
@@ -619,8 +605,6 @@ namespace DatasmithImporterImpl
 	{
 		TRACE_CPUPROFILER_EVENT_SCOPE(DatasmithImporterImpl::PublicizeComponent);
 
-		SafetyGarbageCollect();
-
 		if ( !SourceComponent.HasAnyFlags( RF_Transient | RF_TextExportTransient ) )
 		{
 			if (!DestinationComponent || DestinationComponent->IsPendingKillOrUnreachable())
@@ -925,6 +909,7 @@ void FDatasmithImporter::ImportStaticMeshes( FDatasmithImportContext& ImportCont
 								return nullptr;
 							}
 
+							TRACE_CPUPROFILER_EVENT_SCOPE(LoadStaticMesh);
 							TUniquePtr<FDatasmithMeshElementPayload> MeshPayload = MakeUnique<FDatasmithMeshElementPayload>();
 							return ImportContext.SceneTranslator->LoadStaticMesh(MeshElement, *MeshPayload) ? MeshPayload.Release() : nullptr;
 						}
@@ -1009,6 +994,7 @@ UStaticMesh* FDatasmithImporter::ImportStaticMesh( FDatasmithImportContext& Impo
 		FDatasmithMeshElementPayload LocalMeshPayload;
 		if (MeshPayload == nullptr)
 		{
+			TRACE_CPUPROFILER_EVENT_SCOPE(LoadStaticMesh);
 			ImportContext.SceneTranslator->LoadStaticMesh(MeshElement, LocalMeshPayload);
 			MeshPayload = &LocalMeshPayload;
 		}
@@ -1016,6 +1002,15 @@ UStaticMesh* FDatasmithImporter::ImportStaticMesh( FDatasmithImportContext& Impo
 		ImportedStaticMesh = FDatasmithStaticMeshImporter::ImportStaticMesh( MeshElement, *MeshPayload, ImportContext.ObjectFlags & ~RF_Public, ImportContext.Options->BaseOptions.StaticMeshOptions, ImportContext.AssetsContext, ExistingStaticMesh );
 		AdditionalData = MoveTemp(MeshPayload->AdditionalData);
 
+		// Make sure the garbage collector can collect additional data allocated on other thread
+		for (UDatasmithAdditionalData* Data : AdditionalData)
+		{
+			if (Data)
+			{
+				Data->ClearInternalFlags(EInternalObjectFlags::Async);
+			}
+		}
+
 		// Creation of static mesh failed, remove it from the list of importer mesh elements
 		if (ImportedStaticMesh == nullptr)
 		{
@@ -1985,6 +1980,8 @@ void FDatasmithImporter::ImportLevelVariantSets( FDatasmithImportContext& Import
 		return;
 	}
 
+	TRACE_CPUPROFILER_EVENT_SCOPE(FDatasmithImporter::ImportLevelVariantSets);
+
 	FScopedSlowTask Progress( (float)LevelVariantSetsCount, LOCTEXT("ImportingLevelVariantSets", "Importing Level Variant Sets..."), true, *ImportContext.Warn );
 	Progress.MakeDialog(true);
 
