Integrating CL# 9735441 from Dev-Enterprise
Dataprep: Combining operators "Enable Lightmaps UVs" and "Compute Lightmap Resolution" into "Setup Static Lighting"

diff --git a/Engine/Plugins/Enterprise/DatasmithImporter/Source/DatasmithImporter/Private/DatasmithBlueprintLibrary.cpp b/Engine/Plugins/Enterprise/DatasmithImporter/Source/DatasmithImporter/Private/DatasmithBlueprintLibrary.cpp
--- a/Engine/Plugins/Enterprise/DatasmithImporter/Source/DatasmithImporter/Private/DatasmithBlueprintLibrary.cpp
+++ b/Engine/Plugins/Enterprise/DatasmithImporter/Source/DatasmithImporter/Private/DatasmithBlueprintLibrary.cpp
@@ -21,6 +21,7 @@
 #include "Engine/World.h"
 #include "GameFramework/Actor.h"
 #include "HAL/FileManager.h"
+#include "MeshExport.h"
 #include "Misc/PackageName.h"
 #include "PackageTools.h"
 #include "UObject/Package.h"
@@ -30,18 +31,15 @@
 
 #define LOCTEXT_NAMESPACE "DatasmithBlueprintLibrary"
 
+DEFINE_LOG_CATEGORY_STATIC(LogSetupStaticLighting, Log, All);
+
 namespace DatasmithStaticMeshBlueprintLibraryUtil
 {
-	void EnsureLightmapUVsAreAvailable( UStaticMesh* StaticMesh )
+	void EnsureLightmapSourceUVsAreAvailable( UStaticMesh* StaticMesh )
 	{
-		if ( FMeshDescription* MeshDescription = StaticMesh->GetMeshDescription( 0 ) )
+		if ( StaticMesh->GetNumSourceModels() > 0 && StaticMesh->GetSourceModel(0).BuildSettings.bGenerateLightmapUVs )
 		{
-			const bool bAreLightmapUVsAvailable = DatasmithMeshHelper::HasUVData( *MeshDescription, StaticMesh->LightMapCoordinateIndex );
-
-			if ( !bAreLightmapUVsAvailable )
-			{
-				FDatasmithStaticMeshImporter::PreBuildStaticMesh( StaticMesh );
-			}
+			FDatasmithStaticMeshImporter::PreBuildStaticMesh( StaticMesh );			
 		}
 	}
 
@@ -52,6 +50,43 @@ namespace DatasmithStaticMeshBlueprintLibraryUtil
 
 		return Area;
 	}
+
+	//Used for creating a mapping of StaticMeshes and the StaticMeshComponents that references them in the given list.
+	TMap< UStaticMesh*, TSet< UStaticMeshComponent* > > GetStaticMeshComponentMap(const TArray< UObject* >& Objects)
+	{
+		TMap< UStaticMesh*, TSet< UStaticMeshComponent* > > StaticMeshMap;
+
+		for (UObject* Object : Objects)
+		{
+			if (AActor* Actor = Cast< AActor >(Object))
+			{
+				TInlineComponentArray<UStaticMeshComponent*> StaticMeshComponents(Actor);
+				for (UStaticMeshComponent* StaticMeshComponent : StaticMeshComponents)
+				{
+					if (StaticMeshComponent && StaticMeshComponent->GetStaticMesh())
+					{
+						TSet<UStaticMeshComponent*>& Components = StaticMeshMap.FindOrAdd(StaticMeshComponent->GetStaticMesh());
+						Components.Add(StaticMeshComponent);
+					}
+				}
+			}
+			else if (UStaticMeshComponent* StaticMeshComponent = Cast< UStaticMeshComponent >(Object))
+			{
+				if (UStaticMesh* StaticMesh = StaticMeshComponent->GetStaticMesh())
+				{
+					TSet<UStaticMeshComponent*>& Components = StaticMeshMap.FindOrAdd(StaticMesh);
+					Components.Add(StaticMeshComponent);
+				}
+			}
+			else if (UStaticMesh* StaticMesh = Cast< UStaticMesh >(Object))
+			{
+				TSet<UStaticMeshComponent*>& Components = StaticMeshMap.FindOrAdd(StaticMesh);
+				Components.Add(nullptr);
+			}
+		}
+
+		return StaticMeshMap;
+	}
 }
 
 namespace DatasmithBlueprintLibraryImpl
@@ -235,34 +270,83 @@ void UDatasmithSceneElement::DestroyScene()
 	Reset();
 }
 
-void UDatasmithStaticMeshBlueprintLibrary::ComputeLightmapResolution( const TArray< UObject* >& Objects, bool bApplyChanges, float IdealRatio )
+void UDatasmithStaticMeshBlueprintLibrary::SetupStaticLighting(const TArray< UObject* >& Objects, bool bApplyChanges, bool bGenerateLightmapUVs, float LightmapResolutionIdealRatio)
 {
-	TRACE_CPUPROFILER_EVENT_SCOPE(UDatasmithStaticMeshBlueprintLibrary::ComputeLightmapResolution)
-
 	// Collect all the static meshes and static mesh components to compute lightmap resolution for
-	TMap< UStaticMesh*, TSet< UStaticMeshComponent* > > WorkingSet;
+	TMap< UStaticMesh*, TSet< UStaticMeshComponent* > > StaticMeshMap(DatasmithStaticMeshBlueprintLibraryUtil::GetStaticMeshComponentMap(Objects));
 
-	for ( UObject* Object : Objects )
+	for (const auto& StaticMeshPair : StaticMeshMap)
 	{
-		if ( AActor* Actor = Cast< AActor >( Object ) )
+		UStaticMesh* StaticMesh = StaticMeshPair.Key;
+
+		if (bApplyChanges)
 		{
-			TInlineComponentArray<UStaticMeshComponent*> StaticMeshComponents(Actor);
-			for (UStaticMeshComponent* StaticMeshComponent : StaticMeshComponents)
+			StaticMesh->Modify();
+		}
+
+		for (int32 LODIndex = 0; LODIndex < StaticMesh->GetNumSourceModels(); ++LODIndex)
+		{
+			FStaticMeshSourceModel& SourceModel = StaticMesh->GetSourceModel(LODIndex);
+			const bool bDidChangeSettings = SourceModel.BuildSettings.bGenerateLightmapUVs != bGenerateLightmapUVs;
+			SourceModel.BuildSettings.bGenerateLightmapUVs = bGenerateLightmapUVs;
+
+			if (LODIndex == 0)
 			{
-				if(StaticMeshComponent && StaticMeshComponent->GetStaticMesh())
+				int32 MaxBiggestUVChannel = Lightmass::MAX_TEXCOORDS;
+				
+				if (const FMeshDescription* MeshDescription = SourceModel.MeshDescription.Get())
 				{
-					TSet<UStaticMeshComponent*>& Components = WorkingSet.FindOrAdd( StaticMeshComponent->GetStaticMesh() );
-					Components.Add( StaticMeshComponent );
+					FStaticMeshConstAttributes Attributes(*MeshDescription);
+				
+					// 3 is the maximum that lightmass accept. Defined in MeshExport.h : MAX_TEXCOORDS .
+					MaxBiggestUVChannel = FMath::Min(MaxBiggestUVChannel, Attributes.GetVertexInstanceUVs().GetNumIndices() - 1);
+				}
+
+				if (bGenerateLightmapUVs)
+				{
+					const int32 GeneratedLightmapChannel = SourceModel.BuildSettings.DstLightmapIndex;
+					
+					if (GeneratedLightmapChannel < Lightmass::MAX_TEXCOORDS)
+					{
+						StaticMesh->LightMapCoordinateIndex = GeneratedLightmapChannel;
+					}
+					else
+					{
+						UE_LOG(LogSetupStaticLighting, Warning, TEXT("Could not complete the static lighting setup for static mesh %s as the generated lightmap UV is set to be in channel #%i while the maximum lightmap channel is %i"), *StaticMesh->GetName(), GeneratedLightmapChannel, Lightmass::MAX_TEXCOORDS);
+						break;
+					}
+				}
+				else if (StaticMesh->LightMapCoordinateIndex > MaxBiggestUVChannel && bDidChangeSettings)
+				{
+					// If we are not generating the lightmap anymore make sure we are selecting a valid lightmap index.
+					StaticMesh->LightMapCoordinateIndex = MaxBiggestUVChannel;
 				}
 			}
 		}
-		else if ( UStaticMesh* StaticMesh = Cast< UStaticMesh >( Object ) )
-		{
-			TSet<UStaticMeshComponent*>& Components = WorkingSet.FindOrAdd(StaticMesh);
-			Components.Add( nullptr );
-		}
 	}
 
+	// Compute the lightmap resolution, do not apply the changes so that the computation is done on multiple threads
+	// We'll directly call PostEditChange() at the end of the function so that we also get the StaticLightingSetup changes.
+	ComputeLightmapResolution(StaticMeshMap, /* bApplyChange */false, LightmapResolutionIdealRatio);
+
+	for (const auto& StaticMeshPair : StaticMeshMap)
+	{
+		StaticMeshPair.Key->PostEditChange();
+	}
+}
+
+void UDatasmithStaticMeshBlueprintLibrary::ComputeLightmapResolution(const TArray< UObject* >& Objects, bool bApplyChanges, float IdealRatio)
+{
+	TRACE_CPUPROFILER_EVENT_SCOPE(UDatasmithStaticMeshBlueprintLibrary::ComputeLightmapResolution)
+
+	// Collect all the static meshes and static mesh components to compute lightmap resolution for
+	TMap< UStaticMesh*, TSet< UStaticMeshComponent* > > StaticMeshMap(DatasmithStaticMeshBlueprintLibraryUtil::GetStaticMeshComponentMap(Objects));
+
+	ComputeLightmapResolution(StaticMeshMap, bApplyChanges, IdealRatio);
+}
+
+void UDatasmithStaticMeshBlueprintLibrary::ComputeLightmapResolution(const TMap< UStaticMesh*, TSet< UStaticMeshComponent* > >& StaticMeshMap, bool bApplyChanges, float IdealRatio)
+{
 	// The actual work
 	auto Compute = [&](UStaticMesh* StaticMesh, const TSet<UStaticMeshComponent*>& Components)
 	{
@@ -326,20 +410,31 @@ void UDatasmithStaticMeshBlueprintLibrary::ComputeLightmapResolution( const TArr
 	// If no need to notify changes, multi-thread the computing
 	if(!bApplyChanges)
 	{
-		TArray< UStaticMesh* > WorkingSetKeys;
-		WorkingSet.GenerateKeyArray( WorkingSetKeys );
+		TArray< UStaticMesh* > StaticMeshes;
+		StaticMeshMap.GenerateKeyArray( StaticMeshes );
 
 		// Start with the biggest mesh first to help balancing tasks on threads
 		Algo::SortBy(
-			WorkingSetKeys,
+			StaticMeshes,
 			[](const UStaticMesh* Mesh){ return Mesh->IsMeshDescriptionValid(0) ? Mesh->GetMeshDescription(0)->Vertices().Num() : 0; },
 			TGreater<>()
 		);
 
-		ParallelFor( WorkingSetKeys.Num(),
+		ParallelFor(StaticMeshes.Num(),
+			[&](int32 Index)
+			{
+				// We need to ensure the source UVs for generated lightmaps are available before generating then in the UStaticMesh::BatchBuild().
+				DatasmithStaticMeshBlueprintLibraryUtil::EnsureLightmapSourceUVsAreAvailable(StaticMeshes[Index]);
+			},
+			EParallelForFlags::Unbalanced
+		);
+
+		UStaticMesh::BatchBuild( StaticMeshes, true);
+
+		ParallelFor( StaticMeshes.Num(),
 			[&]( int32 Index )
 			{
-				Compute(  WorkingSetKeys[Index], WorkingSet[WorkingSetKeys[Index]] );
+				Compute( StaticMeshes[Index], StaticMeshMap[StaticMeshes[Index]] );
 			},
 			EParallelForFlags::Unbalanced
 		);
@@ -347,7 +442,7 @@ void UDatasmithStaticMeshBlueprintLibrary::ComputeLightmapResolution( const TArr
 	// Do not take any chance, compute sequentially
 	else
 	{
-		for(TPair< UStaticMesh*, TSet< UStaticMeshComponent* > >& Entry : WorkingSet)
+		for(const TPair< UStaticMesh*, TSet< UStaticMeshComponent* > >& Entry : StaticMeshMap)
 		{
 			Compute(  Entry.Key, Entry.Value );
 		}
@@ -356,82 +451,63 @@ void UDatasmithStaticMeshBlueprintLibrary::ComputeLightmapResolution( const TArr
 
 int32 UDatasmithStaticMeshBlueprintLibrary::ComputeLightmapResolution(UStaticMesh* StaticMesh, float IdealRatio, const FVector& StaticMeshScale)
 {
-	if(StaticMesh == nullptr)
+	if(StaticMesh == nullptr || !StaticMesh->HasValidRenderData())
 	{
 		return 0;
 	}
 
-	DatasmithStaticMeshBlueprintLibraryUtil::EnsureLightmapUVsAreAvailable( StaticMesh );
+	const FRawStaticIndexBuffer& IndexBuffer = StaticMesh->RenderData->LODResources[0].IndexBuffer;
+	const FPositionVertexBuffer& PositionBuffer = StaticMesh->RenderData->LODResources[0].VertexBuffers.PositionVertexBuffer;
+	const FStaticMeshVertexBuffer& VertexBuffer = StaticMesh->RenderData->LODResources[0].VertexBuffers.StaticMeshVertexBuffer;
 
-	if ( const FMeshDescription* MeshDescription = StaticMesh->GetMeshDescription( 0 ) )
+	if (VertexBuffer.GetNumTexCoords() <= (uint32)StaticMesh->LightMapCoordinateIndex)
 	{
-		FStaticMeshConstAttributes Attributes(*MeshDescription);
+		return 0;
+	}
 
-		// Compute the mesh UV density, based on FStaticMeshRenderData::ComputeUVDensities, except that we're working on a MeshDescription
-		const TVertexAttributesConstRef< FVector > VertexPositions = Attributes.GetVertexPositions();
-		const TVertexInstanceAttributesConstRef< FVector2D > VertexUVs = Attributes.GetVertexInstanceUVs();
+	// Compute the mesh UV density, based on FStaticMeshRenderData::ComputeUVDensities, except that we're only working the Lightmap UV.
+	TArray< FVector2D > PolygonAreas;
+	const int32 NumberOfTriangles = IndexBuffer.GetNumIndices() / 3;
+	for (int32 TriangleIndex = 0; TriangleIndex < NumberOfTriangles; ++TriangleIndex)
+	{
+		FVector VertexPosition[3];
+		FVector2D LightmapUVs[3];
 
-		if ( VertexUVs.GetNumElements() <= StaticMesh->LightMapCoordinateIndex )
+		for (int32 CornerIndex = 0; CornerIndex < 3; ++CornerIndex)
 		{
-			return 0;
+			uint32 VertexIndex = IndexBuffer.GetIndex(TriangleIndex * 3 + CornerIndex);
+			VertexPosition[CornerIndex] = PositionBuffer.VertexPosition(VertexIndex);
+			LightmapUVs[CornerIndex] = VertexBuffer.GetVertexUV(VertexIndex, StaticMesh->LightMapCoordinateIndex);
 		}
 
-		float MeshArea = 0.f;
-		float MeshUVArea = 0.f;
+		const float PolygonArea = DatasmithStaticMeshBlueprintLibraryUtil::ParallelogramArea(VertexPosition[0], VertexPosition[1], VertexPosition[2]);
+		const float PolygonUVArea = DatasmithStaticMeshBlueprintLibraryUtil::ParallelogramArea(FVector(LightmapUVs[0], 0.f), FVector(LightmapUVs[1], 0.f), FVector(LightmapUVs[2], 0.f));
 
-		TArray< FVector2D > PolygonAreas;
-
-		for ( const FPolygonID PolygonID : MeshDescription->Polygons().GetElementIDs() )
-		{
-			float PolygonArea = 0.f;
-			float PolygonUVArea = 0.f;
-
-			const TArray< FTriangleID >& PolygonTriangleIDs = MeshDescription->GetPolygonTriangleIDs( PolygonID );
-			for ( const FTriangleID TriangleID : PolygonTriangleIDs )
-			{
-				FVector VertexPosition[3];
-				FVector2D LightmapUVs[3];
-
-				for ( int32 CornerIndex = 0; CornerIndex < 3; ++CornerIndex )
-				{
-					const FVertexInstanceID VertexInstanceID = MeshDescription->GetTriangleVertexInstance( TriangleID, CornerIndex );
-					VertexPosition[ CornerIndex ] = VertexPositions[ MeshDescription->GetVertexInstanceVertex( VertexInstanceID ) ] * StaticMeshScale;
-
-					LightmapUVs[ CornerIndex ] = VertexUVs.Get( VertexInstanceID, StaticMesh->LightMapCoordinateIndex );
-				}
-
-				PolygonArea += DatasmithStaticMeshBlueprintLibraryUtil::ParallelogramArea( VertexPosition[0],  VertexPosition[1],  VertexPosition[2] );
-				PolygonUVArea += DatasmithStaticMeshBlueprintLibraryUtil::ParallelogramArea( FVector( LightmapUVs[0], 0.f ),  FVector( LightmapUVs[1], 0.f ),  FVector( LightmapUVs[2], 0.f ) );
-			}
-
-			PolygonAreas.Emplace( FMath::Sqrt( PolygonArea ), FMath::Sqrt(PolygonArea / PolygonUVArea ) );
-		}
-
-		Algo::Sort( PolygonAreas, []( const FVector2D& ElementA, const FVector2D& ElementB )
-		{
-			return ElementA[1] < ElementB[1];
-		} );
+		PolygonAreas.Emplace(FMath::Sqrt(PolygonArea), FMath::Sqrt(PolygonArea / PolygonUVArea));
+	}
 
-		float WeightedUVDensity = 0.f;
-		float Weight = 0.f;
+	Algo::Sort( PolygonAreas, []( const FVector2D& ElementA, const FVector2D& ElementB )
+	{
+		return ElementA[1] < ElementB[1];
+	} );
 
-		// Remove 10% of higher and lower texel factors.
-		const int32 Threshold = FMath::FloorToInt( 0.1f * (float)PolygonAreas.Num() );
-		for (int32 Index = Threshold; Index < PolygonAreas.Num() - Threshold; ++Index)
-		{
-			WeightedUVDensity += PolygonAreas[ Index ][ 1 ] * PolygonAreas[ Index ][ 0 ];
-			Weight += PolygonAreas[ Index ][ 0 ];
-		}
+	float WeightedUVDensity = 0.f;
+	float Weight = 0.f;
 
-		float UVDensity = WeightedUVDensity / Weight;
+	// Remove 10% of higher and lower texel factors.
+	const int32 Threshold = FMath::FloorToInt( 0.1f * (float)PolygonAreas.Num() );
+	for (int32 Index = Threshold; Index < PolygonAreas.Num() - Threshold; ++Index)
+	{
+		WeightedUVDensity += PolygonAreas[ Index ][ 1 ] * PolygonAreas[ Index ][ 0 ];
+		Weight += PolygonAreas[ Index ][ 0 ];
+	}
 
-		int32 LightmapResolution = FMath::FloorToInt( UVDensity * IdealRatio );
+	float UVDensity = WeightedUVDensity / Weight;
 
-		// Ensure that LightmapResolution is a factor of 4
-		return FMath::Max( LightmapResolution + 3 & ~3, 4 );
-	}
+	int32 LightmapResolution = FMath::FloorToInt( UVDensity * IdealRatio );
 
-	return 0;
+	// Ensure that LightmapResolution is a factor of 4
+	return FMath::Max( LightmapResolution + 3 & ~3, 4 );
 }
 
 FDatasmithImportFactoryCreateFileResult::FDatasmithImportFactoryCreateFileResult()
diff --git a/Engine/Plugins/Enterprise/DatasmithImporter/Source/DatasmithImporter/Private/DatasmithDataprepOperation.cpp b/Engine/Plugins/Enterprise/DatasmithImporter/Source/DatasmithImporter/Private/DatasmithDataprepOperation.cpp
--- a/Engine/Plugins/Enterprise/DatasmithImporter/Source/DatasmithImporter/Private/DatasmithDataprepOperation.cpp
+++ b/Engine/Plugins/Enterprise/DatasmithImporter/Source/DatasmithImporter/Private/DatasmithDataprepOperation.cpp
@@ -7,20 +7,14 @@
 
 #define LOCTEXT_NAMESPACE "DatasmithDataprepOperation"
 
-void UDatasmithComputeLightmapResolutionOperation::OnExecution_Implementation(const FDataprepContext& InContext)
+void UDataprepSetupStaticLightingOperation::OnExecution_Implementation(const FDataprepContext& InContext)
 {
-	// Collect start time to log amount of time spent to import incoming file
-	uint64 StartTime = FPlatformTime::Cycles64();
-	int32 ObjectsCount = InContext.Objects.Num();
+#ifdef LOG_TIME
+	DataprepOperationTime::FTimeLogger TimeLogger(TEXT("SetupStaticLighting"), [&](FText Text) { this->LogInfo(Text); });
+#endif
 
-	UDatasmithStaticMeshBlueprintLibrary::ComputeLightmapResolution( InContext.Objects, false, IdealRatio );
-
-	// Log time spent to import incoming file in minutes and seconds
-	double ElapsedSeconds = FPlatformTime::ToSeconds64(FPlatformTime::Cycles64() - StartTime);
-
-	int ElapsedMin = int(ElapsedSeconds / 60.0);
-	ElapsedSeconds -= 60.0 * (double)ElapsedMin;
-	UE_LOG( LogDatasmithImport, Log, TEXT("Computation of lightmap resolution of %d object(s) took [%d min %.3f s]"), ObjectsCount, ElapsedMin, ElapsedSeconds );
+	// Execute operation
+	UDatasmithStaticMeshBlueprintLibrary::SetupStaticLighting(InContext.Objects, false, bEnableLightmapUVGeneration, LightmapResolutionIdealRatio);
 }
 
 #undef LOCTEXT_NAMESPACE
diff --git a/Engine/Plugins/Enterprise/DatasmithImporter/Source/DatasmithImporter/Private/DatasmithDataprepOperation.h b/Engine/Plugins/Enterprise/DatasmithImporter/Source/DatasmithImporter/Private/DatasmithDataprepOperation.h
--- a/Engine/Plugins/Enterprise/DatasmithImporter/Source/DatasmithImporter/Private/DatasmithDataprepOperation.h
+++ b/Engine/Plugins/Enterprise/DatasmithImporter/Source/DatasmithImporter/Private/DatasmithDataprepOperation.h
@@ -8,20 +8,26 @@
 
 #include "DatasmithDataprepOperation.generated.h"
 
-UCLASS(Experimental, Category = MeshOperation, Meta = (DisplayName="Compute Lightmap Resolution", ToolTip = "For each static mesh to process, recompute the lilghtmap resolution based on the specified ratio") )
-class UDatasmithComputeLightmapResolutionOperation : public UDataprepOperation
+UCLASS(Experimental, Category = LightmapOptions, Meta = (DisplayName = "Setup Static Lighting", ToolTip = "For each static mesh to process, setup the settings to enable lightmap UVs generation and compute the lightmap resolution."))
+class UDataprepSetupStaticLightingOperation : public UDataprepOperation
 {
 	GENERATED_BODY()
 
-	UDatasmithComputeLightmapResolutionOperation()
-		: IdealRatio( 0.2f )
+		UDataprepSetupStaticLightingOperation()
+		: bEnableLightmapUVGeneration(true),
+		LightmapResolutionIdealRatio(0.2f)
 	{
 	}
 
 public:
-	UPROPERTY(EditAnywhere, BlueprintReadWrite, Category = LightmapOptions )
-	float IdealRatio;
+	// The value to set for the generate lightmap uvs flag on each static mesh
+	UPROPERTY(EditAnywhere, BlueprintReadWrite, Category = LightmapOptions, meta = (DisplayName = "Enable Lightmap UV Generation", ToolTip = "Enable the lightmap UV generation."))
+	bool bEnableLightmapUVGeneration;
 
+	UPROPERTY(EditAnywhere, BlueprintReadWrite, Category = LightmapOptions, meta = (DisplayName = "Resolution Ideal Ratio", ToolTip = "The ratio used to compute the resolution of the lightmap."))
+	float LightmapResolutionIdealRatio;
+
+protected:
 	//~ Begin UDataprepOperation Interface
 public:
 	virtual FText GetCategory_Implementation() const override
@@ -32,4 +38,4 @@ public:
 protected:
 	virtual void OnExecution_Implementation(const FDataprepContext& InContext) override;
 	//~ End UDataprepOperation Interface
-};
+};
\ No newline at end of file
diff --git a/Engine/Plugins/Enterprise/DatasmithImporter/Source/DatasmithImporter/Public/DatasmithBlueprintLibrary.h b/Engine/Plugins/Enterprise/DatasmithImporter/Source/DatasmithImporter/Public/DatasmithBlueprintLibrary.h
--- a/Engine/Plugins/Enterprise/DatasmithImporter/Source/DatasmithImporter/Public/DatasmithBlueprintLibrary.h
+++ b/Engine/Plugins/Enterprise/DatasmithImporter/Source/DatasmithImporter/Public/DatasmithBlueprintLibrary.h
@@ -104,6 +104,19 @@ public:
 	UFUNCTION(BlueprintCallable, Category = "Datasmith | Static Mesh")
 	static void ComputeLightmapResolution( const TArray< UObject* >& Objects, bool bApplyChanges, float IdealRatio = 0.2f );
 
+	/**
+	 * Setup the Lightmap UVs settings to enable or disable the lightmap generation on the static meshes found in the Assets list
+	 *
+	 * @param	Assets							List of objects to set the generate lightmap uvs flag on. Only Static Meshes and Static Mesh Components will be affected.
+	 * @param	bApplyChanges					Indicates if changes must be apply or not.
+	 * @param	bGenerateLightmapUVs			The value to set for the generate lightmap uvs flag.
+	 * @param	LightmapResolutionIdealRatio	The desired lightmap density ratio
+	 */
+	UFUNCTION(BlueprintCallable, Category = "Datasmith | Static Mesh")
+	static void SetupStaticLighting(const TArray< UObject* >& Objects, bool bApplyChanges, bool bGenerateLightmapUVs, float LightmapResolutionIdealRatio = 0.2f );
+
 private:
+	static void ComputeLightmapResolution(const TMap< UStaticMesh*, TSet< UStaticMeshComponent* > >& StaticMeshMap, bool bApplyChanges, float IdealRatio = 0.2f);
+
 	static int32 ComputeLightmapResolution(UStaticMesh* StaticMesh, float IdealRatio, const FVector& StaticMeshScale);
 };
diff --git a/Engine/Plugins/Experimental/DataPrepEditor/Source/DataPrepLibraries/Private/DataPrepOperationsLibrary.cpp b/Engine/Plugins/Experimental/DataPrepEditor/Source/DataPrepLibraries/Private/DataPrepOperationsLibrary.cpp
--- a/Engine/Plugins/Experimental/DataPrepEditor/Source/DataPrepLibraries/Private/DataPrepOperationsLibrary.cpp
+++ b/Engine/Plugins/Experimental/DataPrepEditor/Source/DataPrepLibraries/Private/DataPrepOperationsLibrary.cpp
@@ -71,6 +71,13 @@ namespace DataprepOperationsLibraryUtil
 			{
 				SelectedMeshes.Add( StaticMesh );
 			}
+			else if ( UStaticMeshComponent* StaticMeshComponent = Cast<UStaticMeshComponent>(Object) )
+			{
+				if (UStaticMesh* StaticMesh = StaticMeshComponent->GetStaticMesh())
+				{
+					SelectedMeshes.Add(StaticMesh);
+				}
+			}
 			else if (AActor* Actor = Cast<AActor>(Object) )
 			{
 				TInlineComponentArray<UStaticMeshComponent*> StaticMeshComponents( Actor );
@@ -450,41 +457,6 @@ void UDataprepOperationsLibrary::SetConvexDecompositionCollision(const TArray<UO
 	}
 }
 
-void UDataprepOperationsLibrary::SetGenerateLightmapUVs( const TArray< UObject* >& Assets, bool bGenerateLightmapUVs, TArray<UObject*>& ModifiedObjects )
-{
-	TSet<UStaticMesh*> SelectedMeshes = DataprepOperationsLibraryUtil::GetSelectedMeshes(Assets);
-
-	for (UStaticMesh* StaticMesh : SelectedMeshes)
-	{
-		if (StaticMesh)
-		{
-			bool bDidChangeSettings = false;
-
-			// 3 is the maximum that lightmass accept
-			int32 MinBiggestUVChannel = 3;
-			for ( FStaticMeshSourceModel& SourceModel : StaticMesh->GetSourceModels() )
-			{
-				bDidChangeSettings |= SourceModel.BuildSettings.bGenerateLightmapUVs != bGenerateLightmapUVs;
-				SourceModel.BuildSettings.bGenerateLightmapUVs = bGenerateLightmapUVs;
-				if( const FMeshDescription* MeshDescription = SourceModel.MeshDescription.Get() )
-				{
-					FStaticMeshConstAttributes Attributes(*MeshDescription);
-					int32 UVChannelCount = Attributes.GetVertexInstanceUVs().GetNumIndices();
-					MinBiggestUVChannel = FMath::Min( MinBiggestUVChannel, UVChannelCount - 1 );
-				}
-			}
-
-			if ( StaticMesh->LightMapCoordinateIndex > MinBiggestUVChannel && bDidChangeSettings )
-			{
-				// Correct the coordinate index if it was invalid
-				StaticMesh->LightMapCoordinateIndex = MinBiggestUVChannel;
-			}
-
-			ModifiedObjects.Add( StaticMesh );
-		}
-	}
-}
-
 void UDataprepOperationsLibrary::SubstituteMaterial(const TArray<UObject*>& SelectedObjects, const FString& MaterialSearch, EEditorScriptingStringMatchType StringMatch, UMaterialInterface* MaterialSubstitute)
 {
 	TArray<UMaterialInterface*> MaterialsUsed = DataprepOperationsLibraryUtil::GetUsedMaterials(SelectedObjects);
diff --git a/Engine/Plugins/Experimental/DataPrepEditor/Source/DataPrepLibraries/Private/DataPrepOperationsLibrary.h b/Engine/Plugins/Experimental/DataPrepEditor/Source/DataPrepLibraries/Private/DataPrepOperationsLibrary.h
--- a/Engine/Plugins/Experimental/DataPrepEditor/Source/DataPrepLibraries/Private/DataPrepOperationsLibrary.h
+++ b/Engine/Plugins/Experimental/DataPrepEditor/Source/DataPrepLibraries/Private/DataPrepOperationsLibrary.h
@@ -191,15 +191,6 @@ public:
 	UFUNCTION(BlueprintCallable, Category = "Dataprep | Operation")
 	static void SetConvexDecompositionCollision(const TArray<UObject*>& SelectedObjects, int32 HullCount, int32 MaxHullVerts, int32 HullPrecision, TArray<UObject*>& ModifiedObjects);
 
-	/**
-	 * Sets the Generate Lightmap UVs flag on the static meshes found in the Assets list
-	 *
-	 * @param	Assets					List of assets to set the generate lightmap uvs flag on. Only Static Meshes will be affected.
-	 * @param	bGenerateLightmapUVs	The value to set for the generate lightmap uvs flag.
-	 */
-	UFUNCTION(BlueprintCallable, Category = "Dataprep | Operation")
-	static void SetGenerateLightmapUVs( const TArray< UObject* >& Assets, bool bGenerateLightmapUVs, TArray<UObject*>& ModifiedObjects );
-
 	/**
 	 * Replaces designated materials in all or specific content folders with specific ones
 	 * @param SelectedObjects: Objects to consider for the substitution
diff --git a/Engine/Plugins/Experimental/DataPrepEditor/Source/DataPrepLibraries/Private/DataprepOperations.cpp b/Engine/Plugins/Experimental/DataPrepEditor/Source/DataPrepLibraries/Private/DataprepOperations.cpp
--- a/Engine/Plugins/Experimental/DataPrepEditor/Source/DataPrepLibraries/Private/DataprepOperations.cpp
+++ b/Engine/Plugins/Experimental/DataPrepEditor/Source/DataPrepLibraries/Private/DataprepOperations.cpp
@@ -147,22 +147,6 @@ void UDataprepSetConvexDecompositionCollisionOperation::OnExecution_Implementati
 	}
 }
 
-void UDataprepSetGenerateLightmapUVsOperation::OnExecution_Implementation(const FDataprepContext& InContext)
-{
-#ifdef LOG_TIME
-	DataprepOperationTime::FTimeLogger TimeLogger( TEXT("SetGenerateLightmapUVs"), [&]( FText Text) { this->LogInfo( Text ); });
-#endif
-
-	// Execute operation
-	TArray<UObject*> ModifiedStaticMeshes;
-	UDataprepOperationsLibrary::SetGenerateLightmapUVs( InContext.Objects, bGenerateLightmapUVs, ModifiedStaticMeshes );
-
-	if(ModifiedStaticMeshes.Num() > 0)
-	{
-		AssetsModified( MoveTemp( ModifiedStaticMeshes ) );
-	}
-}
-
 void UDataprepSetMobilityOperation::OnExecution_Implementation(const FDataprepContext& InContext)
 {
 #ifdef LOG_TIME
diff --git a/Engine/Plugins/Experimental/DataPrepEditor/Source/DataPrepLibraries/Private/DataprepOperations.h b/Engine/Plugins/Experimental/DataPrepEditor/Source/DataPrepLibraries/Private/DataprepOperations.h
--- a/Engine/Plugins/Experimental/DataPrepEditor/Source/DataPrepLibraries/Private/DataprepOperations.h
+++ b/Engine/Plugins/Experimental/DataPrepEditor/Source/DataPrepLibraries/Private/DataprepOperations.h
@@ -163,34 +163,6 @@ protected:
 	//~ End UDataprepOperation Interface
 };
 
-UCLASS(Experimental, Category = MeshOperation, Meta = (DisplayName="Enable Lightmap UVs", ToolTip = "For each static mesh to process, enable or disable the generation of lightmap UVs") )
-class UDataprepSetGenerateLightmapUVsOperation : public UDataprepOperation
-{
-	GENERATED_BODY()
-
-	UDataprepSetGenerateLightmapUVsOperation()
-		: bGenerateLightmapUVs(true)
-	{
-	}
-
-public:
-	// The value to set for the generate lightmap uvs flag on each static mesh
-	UPROPERTY(EditAnywhere, BlueprintReadWrite, Category = MeshOperation, meta = (ToolTip = "Maximum number of convex pieces that will be created"))
-	bool bGenerateLightmapUVs;
-
-protected:
-	//~ Begin UDataprepOperation Interface
-public:
-	virtual FText GetCategory_Implementation() const override
-	{
-		return FDataprepOperationCategories::MeshOperation;
-	}
-
-protected:
-	virtual void OnExecution_Implementation(const FDataprepContext& InContext) override;
-	//~ End UDataprepOperation Interface
-};
-
 UCLASS(Experimental, Category = ActorOperation, Meta = (DisplayName="Set Mobility", ToolTip = "For each mesh actor to process, update its mobilty with the selected value") )
 class UDataprepSetMobilityOperation : public UDataprepOperation
 {
