Fix possible invalid name in datasmith materials

diff --git a/Engine/Plugins/Enterprise/DatasmithImporter/Source/DatasmithImporter/Private/DatasmithMaterialExpressions.cpp b/Engine/Plugins/Enterprise/DatasmithImporter/Source/DatasmithImporter/Private/DatasmithMaterialExpressions.cpp
--- a/Engine/Plugins/Enterprise/DatasmithImporter/Source/DatasmithImporter/Private/DatasmithMaterialExpressions.cpp
+++ b/Engine/Plugins/Enterprise/DatasmithImporter/Source/DatasmithImporter/Private/DatasmithMaterialExpressions.cpp
@@ -149,9 +149,9 @@ namespace
 
 	FString GenerateUniqueMaterialName( const TCHAR* MaterialLabel, FDatasmithUniqueNameProvider& UniqueNameProvider )
 	{
-		FString Label = MaterialLabel;
+		FString Label = ObjectTools::SanitizeObjectName( MaterialLabel );
 		// Generate unique name from label if valid, otherwise return element's name as it is unique
-		return  UniqueNameProvider.GenerateUniqueName( Label );
+		return UniqueNameProvider.GenerateUniqueName( Label );
 	}
 }
 
