Add option to propagate visibility animations to children #jira UE-81901 #rb Anousack.Kitisa, JeanMichel.Dignard

diff --git a/Engine/Plugins/Enterprise/DatasmithFBXImporter/Source/DatasmithVREDTranslator/Private/DatasmithVREDImporter.cpp b/Engine/Plugins/Enterprise/DatasmithFBXImporter/Source/DatasmithVREDTranslator/Private/DatasmithVREDImporter.cpp
--- a/Engine/Plugins/Enterprise/DatasmithFBXImporter/Source/DatasmithVREDTranslator/Private/DatasmithVREDImporter.cpp
+++ b/Engine/Plugins/Enterprise/DatasmithFBXImporter/Source/DatasmithVREDTranslator/Private/DatasmithVREDImporter.cpp
@@ -1770,6 +1770,7 @@ TSharedPtr<IDatasmithLevelSequenceElement> FDatasmithVREDImporter::ConvertAnimBl
 
 				if (VisibilityAnimation.IsValid())
 				{
+					VisibilityAnimation->SetPropagateToChildren(true);
 					PopulateVisibilityAnimation(VisibilityAnimation.ToSharedRef().Get(), Curves[0]);
 				}
 			}
diff --git a/Engine/Plugins/Enterprise/DatasmithImporter/Source/DatasmithImporter/Private/DatasmithLevelSequenceImporter.cpp b/Engine/Plugins/Enterprise/DatasmithImporter/Source/DatasmithImporter/Private/DatasmithLevelSequenceImporter.cpp
--- a/Engine/Plugins/Enterprise/DatasmithImporter/Source/DatasmithImporter/Private/DatasmithLevelSequenceImporter.cpp
+++ b/Engine/Plugins/Enterprise/DatasmithImporter/Source/DatasmithImporter/Private/DatasmithLevelSequenceImporter.cpp
@@ -115,6 +115,28 @@ namespace DatasmithLevelSequenceImporterImpl
 		return ObjectBinding;
 	}
 
+	// Creates duplicates of OriginalTrack for every child (recursively) of ParentActor, adding them to the MovieScene
+	void PropagateVisibilityTrackRecursive(AActor* ParentActor, int32 AnimIndex, ULevelSequence* LevelSequence, UMovieScene* MovieScene, UMovieSceneVisibilityTrack* OriginalTrack)
+	{
+		TArray<AActor*> Children;
+		ParentActor->GetAttachedActors(Children, false);
+		for (AActor* Child : Children)
+		{
+			FGuid ChildBinding = BindActorToLevelSequence(Child, LevelSequence, AnimIndex);
+
+			// Can only have one track of a type per binding
+			if (UMovieSceneVisibilityTrack* ChildVisibilityTrack = MovieScene->FindTrack<UMovieSceneVisibilityTrack>(ChildBinding))
+			{
+				MovieScene->RemoveTrack(*ChildVisibilityTrack);
+			}
+
+			UMovieSceneVisibilityTrack* ClonedParentTrack = DuplicateObject(OriginalTrack, nullptr);
+			MovieScene->AddGivenTrack(ClonedParentTrack, ChildBinding);
+
+			PropagateVisibilityTrackRecursive(Child, ++AnimIndex, LevelSequence, MovieScene, OriginalTrack);
+		}
+	}
+
 	int32 GetTransformValues(const TSharedRef<IDatasmithTransformAnimationElement>& Animation, EDatasmithTransformType TransformType, const FFrameRate& SrcFrameRate, const FFrameRate& DestFrameRate, TArray<FFrameNumber>& FrameNumbers, TArray<float>& XValues, TArray<float>& YValues, TArray<float>& ZValues, FFrameNumber& MinFrameNumber, FFrameNumber& MaxFrameNumber)
 	{
 		int32 NumFrames = Animation->GetFramesCount(TransformType);
@@ -418,6 +440,11 @@ namespace DatasmithLevelSequenceImporterImpl
 		{
 			Channel->Optimize(OptimParams);
 		}
+
+		if (VisibilityAnimation->GetPropagateToChildren())
+		{
+			PropagateVisibilityTrackRecursive(Actor, AnimIndex + 1, LevelSequence, MovieScene, VisibilityTrack);
+		}
 	}
 
 	void PopulateSubsequenceTrack(const TSharedRef<IDatasmithSubsequenceAnimationElement>& SubsequenceAnimation, bool bShouldClearMasterTrack, FDatasmithImportContext& ImportContext, ULevelSequence* LevelSequence, FFrameNumber& MinFrameNumber, FFrameNumber& MaxFrameNumber)
diff --git a/Engine/Source/Developer/Datasmith/DatasmithCore/Private/DatasmithAnimationElementsImpl.cpp b/Engine/Source/Developer/Datasmith/DatasmithCore/Private/DatasmithAnimationElementsImpl.cpp
--- a/Engine/Source/Developer/Datasmith/DatasmithCore/Private/DatasmithAnimationElementsImpl.cpp
+++ b/Engine/Source/Developer/Datasmith/DatasmithCore/Private/DatasmithAnimationElementsImpl.cpp
@@ -22,6 +22,7 @@ FDatasmithTransformAnimationElementImpl::FDatasmithTransformAnimationElementImpl
 
 FDatasmithVisibilityAnimationElementImpl::FDatasmithVisibilityAnimationElementImpl(const TCHAR* InName)
 	: FDatasmithBaseAnimationElementImpl(InName, EDatasmithElementType::Animation, EDatasmithElementAnimationSubType::VisibilityAnimation)
+	, bPropagate(false)
 {
 	CurveInterpMode = EDatasmithCurveInterpMode::Linear;
 }
diff --git a/Engine/Source/Developer/Datasmith/DatasmithCore/Private/DatasmithAnimationElementsImpl.h b/Engine/Source/Developer/Datasmith/DatasmithCore/Private/DatasmithAnimationElementsImpl.h
--- a/Engine/Source/Developer/Datasmith/DatasmithCore/Private/DatasmithAnimationElementsImpl.h
+++ b/Engine/Source/Developer/Datasmith/DatasmithCore/Private/DatasmithAnimationElementsImpl.h
@@ -143,9 +143,20 @@ public:
 		}
 	}
 
+	virtual void SetPropagateToChildren(bool bInPropagate) override
+	{
+		bPropagate = bInPropagate;
+	}
+
+	virtual bool GetPropagateToChildren() const override
+	{
+		return bPropagate;
+	}
+
 private:
 	TArray<FDatasmithVisibilityFrameInfo> Frames;
 	EDatasmithCurveInterpMode CurveInterpMode;
+	bool bPropagate;
 };
 
 class FDatasmithSubsequenceAnimationElementImpl : public FDatasmithBaseAnimationElementImpl< IDatasmithSubsequenceAnimationElement >
diff --git a/Engine/Source/Developer/Datasmith/DatasmithCore/Public/DatasmithAnimationElements.h b/Engine/Source/Developer/Datasmith/DatasmithCore/Public/DatasmithAnimationElements.h
--- a/Engine/Source/Developer/Datasmith/DatasmithCore/Public/DatasmithAnimationElements.h
+++ b/Engine/Source/Developer/Datasmith/DatasmithCore/Public/DatasmithAnimationElements.h
@@ -74,6 +74,12 @@ public:
 
 	/** Remove the frame at the given index from the animation */
 	virtual void RemoveFrame(int32 Index) = 0;
+
+	/** Sets whether this animation will be duplicated to all children (recursively) when imported */
+	virtual void SetPropagateToChildren(bool bPropagate) = 0;
+
+	/** Gets whether this animation will be duplicated to all children (recursively) when imported */
+	virtual bool GetPropagateToChildren() const = 0;
 };
 
 /** IDatasmithSubsequenceAnimationElement holds a reference to a IDatasmithLevelSequenceElement, to be played as a subsequence */
