Removed remaining #ueent_todo
Fixed issue with remaining packages

diff --git a/Engine/Plugins/Enterprise/DatasmithImporter/Source/DatasmithImporter/Private/DatasmithConsumer.cpp b/Engine/Plugins/Enterprise/DatasmithImporter/Source/DatasmithImporter/Private/DatasmithConsumer.cpp
--- a/Engine/Plugins/Enterprise/DatasmithImporter/Source/DatasmithImporter/Private/DatasmithConsumer.cpp
+++ b/Engine/Plugins/Enterprise/DatasmithImporter/Source/DatasmithImporter/Private/DatasmithConsumer.cpp
@@ -97,28 +97,32 @@ bool UDatasmithConsumer::Initialize()
 	UPackage* ParentPackage = CreatePackage( nullptr, *TargetContentFolder );
 	ParentPackage->FullyLoad();
 
+	WorkingScenePtr.Reset( DatasmithScene.Get() );
+
 	// Check if the Datasmith scene is not already in memory
-	if ( !DatasmithScene.IsValid() )
+	if ( !WorkingScenePtr.IsValid() )
 	{
 		FName DatasmithSceneName = MakeUniqueObjectName( ParentPackage, UDatasmithScene::StaticClass(), *( GetName() + TEXT("_DS") ) );
 
 		UPackage* Package = CreatePackage( nullptr, *FPaths::Combine( ParentPackage->GetPathName(), DatasmithSceneName.ToString() ) );
 		Package->FullyLoad();
 
-		DatasmithScene = NewObject< UDatasmithScene >( Package, DatasmithSceneName, GetFlags() | RF_Public | RF_Standalone | RF_Transactional );
+		DatasmithScene = NewObject< UDatasmithScene >( Package, DatasmithSceneName, GetFlags() | RF_Public | RF_Transactional );
 		check( DatasmithScene.IsValid() );
 
+		WorkingScenePtr.Reset( DatasmithScene.Get() );
+
 		FAssetRegistryModule::AssetCreated( DatasmithScene.Get() );
-		DatasmithScene->MarkPackageDirty();
+		WorkingScenePtr->MarkPackageDirty();
 
-		DatasmithScene->AssetImportData = NewObject< UDatasmithSceneImportData >( DatasmithScene.Get(), UDatasmithSceneImportData::StaticClass() );
-		check( DatasmithScene->AssetImportData );
+		WorkingScenePtr->AssetImportData = NewObject< UDatasmithSceneImportData >( DatasmithScene.Get(), UDatasmithSceneImportData::StaticClass() );
+		check( WorkingScenePtr->AssetImportData );
 
-		// Store a dataprep asset pointer into the scene asset in order to be able to later re-execute the dataprep pipeline
+		// Store a Dataprep asset pointer into the scene asset in order to be able to later re-execute the dataprep pipeline
 		UDataprepAssetInterface* DataprepAssetInterface = Cast< UDataprepAssetInterface >( GetOuter() );
 		check( DataprepAssetInterface );
 
-		if ( DatasmithScene->GetClass()->ImplementsInterface(UInterface_AssetUserData::StaticClass()) )
+		if ( WorkingScenePtr->GetClass()->ImplementsInterface(UInterface_AssetUserData::StaticClass()) )
 		{
 			if ( IInterface_AssetUserData* AssetUserDataInterface = Cast< IInterface_AssetUserData >( DatasmithScene.Get() ) )
 			{
@@ -135,10 +139,8 @@ bool UDatasmithConsumer::Initialize()
 		}
 	}
 
-	// #ueent_todo: Find out necessity of namespace for uniqueness of asset's and actor's names
 	if ( !BuildContexts( Context.WorldPtr.Get() ) )
 	{
-		// #ueent_todo: Provide details of why initialization failed
 		return false;
 	}
 
@@ -204,6 +206,7 @@ void UDatasmithConsumer::Reset()
 	ImportContextPtr.Reset();
 	ProgressTaskPtr.Reset();
 	UDataprepContentConsumer::Reset();
+	WorkingScenePtr.Reset(nullptr);
 
 	// Restore previous current level
 	if( PreviousCurrentLevel != nullptr )
@@ -225,18 +228,15 @@ const FText& UDatasmithConsumer::GetDescription() const
 
 bool UDatasmithConsumer::BuildContexts( UWorld* ImportWorld )
 {
-	UDatasmithSceneImportData* ImportData = Cast< UDatasmithSceneImportData >( DatasmithScene->AssetImportData );
-
-	const FString FilePath = FPaths::Combine( FPaths::ProjectIntermediateDir(), ( DatasmithScene->GetName() + TEXT( ".udatasmith" ) ) );
+	const FString FilePath = FPaths::Combine( FPaths::ProjectIntermediateDir(), ( WorkingScenePtr->GetName() + TEXT( ".udatasmith" ) ) );
 
 	ImportContextPtr = MakeUnique< FDatasmithImportContext >( FilePath, false, TEXT("DatasmithImport"), LOCTEXT("DatasmithImportFactoryDescription", "Datasmith") );
 
 	// Update import context with consumer's data
-	ImportContextPtr->Options->BaseOptions = ImportData->BaseOptions;
 	ImportContextPtr->Options->BaseOptions.SceneHandling = EDatasmithImportScene::CurrentLevel;
 	ImportContextPtr->SceneAsset = DatasmithScene.Get();
 	ImportContextPtr->ActorsContext.ImportWorld = ImportWorld;
-	ImportContextPtr->Scene = FDatasmithSceneFactory::CreateScene( *DatasmithScene->GetName() );
+	ImportContextPtr->Scene = FDatasmithSceneFactory::CreateScene( *WorkingScenePtr->GetName() );
 	ImportContextPtr->SceneName = ImportContextPtr->Scene->GetName();
 
 	// Convert all incoming Datasmith scene actors as regular actors
@@ -250,7 +250,7 @@ bool UDatasmithConsumer::BuildContexts( UWorld* ImportWorld )
 	FDatasmithImporterUtils::FillSceneElement( ImportContextPtr->Scene, RootActors );
 
 	// Initialize context
-	FString SceneOuterPath = DatasmithScene->GetOutermost()->GetName();
+	FString SceneOuterPath = WorkingScenePtr->GetOutermost()->GetName();
 	FString RootPath = FPackageName::GetLongPackagePath( SceneOuterPath );
 
 	if ( Algo::Count( RootPath, TEXT('/') ) > 1 )
@@ -263,6 +263,8 @@ bool UDatasmithConsumer::BuildContexts( UWorld* ImportWorld )
 
 	if ( !ImportContextPtr->Init( ImportContextPtr->Scene.ToSharedRef(), RootPath, RF_Public | RF_Standalone | RF_Transactional, GWarn, TSharedPtr< FJsonObject >(), true ) )
 	{
+		FText Message = LOCTEXT( "DatasmithConsumer_BuildContexts", "Initialization of consumer failed" );
+		LogError( Message );
 		return false;
 	}
 
@@ -292,9 +294,6 @@ bool UDatasmithConsumer::BuildContexts( UWorld* ImportWorld )
 	ImportContextPtr->ActorsContext.UniqueNameProvider = FDatasmithActorUniqueLabelProvider();
 	ImportContextPtr->ActorsContext.UniqueNameProvider.PopulateLabelFrom( GWorld );
 
-	// Copy BaseOptions from ImportContext
-	ImportData->BaseOptions.AssetOptions.PackagePath = ImportContextPtr->Options->BaseOptions.AssetOptions.PackagePath;
-
 	// Add assets as if they have been imported using the current import context
 	DatasmithConsumerUtils::AddAssetsToContext( *ImportContextPtr, Context.Assets );
 
@@ -385,7 +384,7 @@ void UDatasmithConsumer::MoveAssets()
 	const FText DialogTitle( LOCTEXT( "DatasmithConsumerDlgTitle", "Warning" ) );
 
 	// Warn user if related Datasmith scene is not in package path and force re-creation of Datasmith scene
-	if( DatasmithScene.IsValid() && !DatasmithScene->GetPathName().StartsWith( TargetContentFolder ) )
+	if( DatasmithScene.IsValid() && !WorkingScenePtr->GetPathName().StartsWith( TargetContentFolder ) )
 	{
 		FText WarningMessage = FText::Format(LOCTEXT("DatasmithConsumer_NoSceneAsset", "Package path {0} different from path previously used, {1}.\nPrevious content will not be updated."), FText::FromString (TargetContentFolder ), FText::FromString ( LastPackagePath ) );
 		FMessageDialog::Open(EAppMsgType::Ok, WarningMessage, &DialogTitle );
@@ -464,7 +463,8 @@ void UDatasmithConsumer::UpdateLevel()
 			}
 			else
 			{
-				// #ueent_todo: Warn user that level could not be created
+				FText Message = LOCTEXT( "DatasmithConsumer_UpdateLevel", "Cannot create level..." );
+				LogWarning( Message );
 				Level = FinalWorld->PersistentLevel;
 			}
 
@@ -522,7 +522,6 @@ namespace DatasmithConsumerUtils
 			const FString SceneActorLabel = SceneActor->GetActorLabel();
 			SceneActor->Rename( nullptr, nullptr, REN_DontCreateRedirectors | REN_NonTransactional );
 
-			// #ueent_todo: is there more to add to the actor element?
 			// Use actor's label instead of name.
 			// Rationale: Datasmith scene actors are created with the same name and label and their name can change when calling SetLabel.
 			TSharedRef< IDatasmithActorElement > RootActorElement = FDatasmithSceneFactory::CreateActor( *SceneActorLabel );
@@ -589,7 +588,6 @@ namespace DatasmithConsumerUtils
 			ActorsToVisit.Append( Children );
 		}
 
-		// #ueent_todo: Find a better way to filter those out.
 		auto IsUnregisteredActor = [&](AActor* Actor)
 		{
 			// Skip non-imported actors
@@ -695,7 +693,6 @@ namespace DatasmithConsumerUtils
 						// Add parent material to ImportedParentMaterials if applicable
 						if ( ParentPath.StartsWith( MaterialInstancePath ) )
 						{
-							// #ueent_todo : Do we want to compute the hash of the material and check its existence?
 							ImportContext.ImportedParentMaterials.Add( MaterialCount, MaterialParent );
 							MaterialCount++;
 
diff --git a/Engine/Plugins/Enterprise/DatasmithImporter/Source/DatasmithImporter/Private/DatasmithConsumer.h b/Engine/Plugins/Enterprise/DatasmithImporter/Source/DatasmithImporter/Private/DatasmithConsumer.h
--- a/Engine/Plugins/Enterprise/DatasmithImporter/Source/DatasmithImporter/Private/DatasmithConsumer.h
+++ b/Engine/Plugins/Enterprise/DatasmithImporter/Source/DatasmithImporter/Private/DatasmithConsumer.h
@@ -68,6 +68,7 @@ private:
 private:
 	TUniquePtr< FDatasmithImportContext > ImportContextPtr;
 	TUniquePtr< FDataprepWorkReporter > ProgressTaskPtr;
+	TStrongObjectPtr< UDatasmithScene > WorkingScenePtr;
 
 	ULevel* PreviousCurrentLevel;
 };
\ No newline at end of file
diff --git a/Engine/Plugins/Enterprise/DatasmithImporter/Source/DatasmithImporter/Private/DatasmithFileProducer.cpp b/Engine/Plugins/Enterprise/DatasmithImporter/Source/DatasmithImporter/Private/DatasmithFileProducer.cpp
--- a/Engine/Plugins/Enterprise/DatasmithImporter/Source/DatasmithImporter/Private/DatasmithFileProducer.cpp
+++ b/Engine/Plugins/Enterprise/DatasmithImporter/Source/DatasmithImporter/Private/DatasmithFileProducer.cpp
@@ -408,7 +408,11 @@ void UDatasmithFileProducer::PreventNameCollision()
 		{
 			if( UObject* Object = Assets[Index].Get() )
 			{
-				PathsToDelete.Add( FPaths::GetPath( Object->GetOutermost()->GetName() ) );
+				// Ensure object's package is transient and not public
+				Object->GetOutermost()->ClearFlags( RF_Public );
+				Object->GetOutermost()->SetFlags( RF_Transient );
+
+				PathsToDelete.Add( Object->GetOutermost()->GetPathName() );
 
 				if( Cast<UStaticMesh>( Object ) != nullptr )
 				{
@@ -777,39 +781,42 @@ void UDatasmithDirProducer::PostInitProperties()
 {
 	UDataprepContentProducer::PostInitProperties();
 
-	if( !HasAnyFlags(RF_ClassDefaultObject | RF_NeedLoad) )
+	if( !HasAnyFlags( RF_ClassDefaultObject | RF_WasLoaded | RF_Transient ) )
 	{
 		FolderPath = FDatasmithFileProducerUtils::SelectDirectory();
+		UpdateName();
 	}
 }
 
 void UDatasmithDirProducer::SetFolderName( const FString& InFolderName )
 {
-	Modify();
+	if(!InFolderName.IsEmpty())
+	{
+		Modify();
 
-	FolderPath = FPaths::ConvertRelativePathToFull( InFolderName );
+		FolderPath = FPaths::ConvertRelativePathToFull( InFolderName );
 
-	FString BaseName = FPaths::IsDrive( InFolderName ) ? TEXT("RootDir") : FPaths::GetBaseFilename( InFolderName );
+		UpdateName();
 
-	if ( Rename( *BaseName, nullptr, REN_Test ) )
-	{
-		Rename( *BaseName, nullptr, REN_DontCreateRedirectors | REN_NonTransactional );
+		OnChanged.Broadcast(this);
 	}
-	else
+}
+
+void UDatasmithDirProducer::UpdateName()
+{
+	if(!FolderPath.IsEmpty())
 	{
-		bool bFoundName = false;
-		for (int32 NameIndex = 1; !bFoundName; ++NameIndex)
+		FString BaseName = FPaths::IsDrive( FolderPath ) ? (FolderPath.Left(1) + TEXT("_Drive")) : (FPaths::GetBaseFilename( FolderPath ) + TEXT("_Dir"));
+
+		// Rename producer to name of file
+		FString CleanName = ObjectTools::SanitizeObjectName( BaseName );
+		if ( !Rename( *CleanName, nullptr, REN_Test ) )
 		{
-			const FString NewName = FString::Printf(TEXT("%s_%d"), *BaseName, NameIndex);
-			if( Rename( *NewName, nullptr, REN_Test ) )
-			{
-				Rename( *NewName, nullptr, REN_DontCreateRedirectors | REN_NonTransactional );
-				bFoundName = true;
-			}
+			CleanName = MakeUniqueObjectName( GetOuter(), GetClass(), *CleanName ).ToString();
 		}
-	}
 
-	OnChanged.Broadcast(this);
+		Rename( *CleanName, nullptr, REN_DontCreateRedirectors | REN_NonTransactional );
+	}
 }
 
 bool UDatasmithDirProducer::Supersede(const UDataprepContentProducer* OtherProducer) const
diff --git a/Engine/Plugins/Enterprise/DatasmithImporter/Source/DatasmithImporter/Private/DatasmithFileProducer.h b/Engine/Plugins/Enterprise/DatasmithImporter/Source/DatasmithImporter/Private/DatasmithFileProducer.h
--- a/Engine/Plugins/Enterprise/DatasmithImporter/Source/DatasmithImporter/Private/DatasmithFileProducer.h
+++ b/Engine/Plugins/Enterprise/DatasmithImporter/Source/DatasmithImporter/Private/DatasmithFileProducer.h
@@ -144,6 +144,9 @@ private:
 	/** Helper function to get all matching files in FolderPath based on extensions set */
 	TSet<FString> GetSetOfFiles() const;
 
+	/** Update the name of the producer based on the directory name */
+	void UpdateName();
+
 private:
 	/** Indicates if ExtensionString contains "*.*" */
 	bool bHasWildCardSearch;
