Fixed build breakage - My apologies...

diff --git a/Engine/Plugins/Enterprise/DatasmithImporter/Source/DatasmithImporter/Private/DatasmithConsumer.h b/Engine/Plugins/Enterprise/DatasmithImporter/Source/DatasmithImporter/Private/DatasmithConsumer.h
--- a/Engine/Plugins/Enterprise/DatasmithImporter/Source/DatasmithImporter/Private/DatasmithConsumer.h
+++ b/Engine/Plugins/Enterprise/DatasmithImporter/Source/DatasmithImporter/Private/DatasmithConsumer.h
@@ -5,6 +5,7 @@
 #include "CoreMinimal.h"
 
 #include "DatasmithImporter.h"
+#include "DatasmithScene.h"
 
 #include "DataPrepContentConsumer.h"
 
