Fixed a crash when canceling a drag and drop of a datasmith scene into the level editor viewport.

diff --git a/Engine/Plugins/Enterprise/DatasmithImporter/Source/DatasmithImporter/Private/ActorFactoryDatasmithScene.cpp b/Engine/Plugins/Enterprise/DatasmithImporter/Source/DatasmithImporter/Private/ActorFactoryDatasmithScene.cpp
--- a/Engine/Plugins/Enterprise/DatasmithImporter/Source/DatasmithImporter/Private/ActorFactoryDatasmithScene.cpp
+++ b/Engine/Plugins/Enterprise/DatasmithImporter/Source/DatasmithImporter/Private/ActorFactoryDatasmithScene.cpp
@@ -16,6 +16,7 @@
 #include "Engine/StaticMeshActor.h"
 #include "Engine/World.h"
 #include "GameFramework/Actor.h"
+#include "UObject/WeakObjectPtrTemplates.h"
 
 #define LOCTEXT_NAMESPACE "ActorFactoryDatasmithScene"
 
@@ -98,9 +99,23 @@ namespace UActorFactoryDatasmithSceneImpl
 
 		ImportContext.SceneAsset = DatasmithScene;
 
+		// The actor might get deleted or become unreachable if the user cancel the import/finalize
+		TWeakObjectPtr<ADatasmithSceneActor> RootObjectAsWeakPtr = RootActor;
+
 		FDatasmithImporter::ImportActors( ImportContext );
 		FDatasmithImporter::FinalizeActors( ImportContext, nullptr );
 
+		RootActor = RootObjectAsWeakPtr.Get();
+
+		if ( RootActor )
+		{
+			// If the root actor is still valid, ensure that it is in a world.
+			if ( !RootActor->GetWorld() )
+			{
+				RootActor = nullptr;
+			}
+		}
+
 		return RootActor;
 	}
 }
