Added a check on existence of file to import and report accordingly

diff --git a/Engine/Plugins/Enterprise/DatasmithImporter/Source/DatasmithImporter/Private/DatasmithFileProducer.cpp b/Engine/Plugins/Enterprise/DatasmithImporter/Source/DatasmithImporter/Private/DatasmithFileProducer.cpp
--- a/Engine/Plugins/Enterprise/DatasmithImporter/Source/DatasmithImporter/Private/DatasmithFileProducer.cpp
+++ b/Engine/Plugins/Enterprise/DatasmithImporter/Source/DatasmithImporter/Private/DatasmithFileProducer.cpp
@@ -101,6 +101,13 @@ bool UDatasmithFileProducer::Initialize()
 		return false;
 	}
 
+	// Check file exists
+	if(!FPaths::FileExists( FilePath ))
+	{
+		LogError( FText::Format( LOCTEXT( "DatasmithFileProducer_NotFound", "File {0} does not exist." ), FText::FromString( FilePath ) ) );
+		return false;
+	}
+
 	UPackage* TransientPackage = NewObject< UPackage >( nullptr, *FPaths::Combine( Context.RootPackagePtr->GetPathName(), *GetName() ), RF_Transient );
 	TransientPackage->FullyLoad();
 
