Fixed a possible crash after copying dataprep asset, changing it's destination folder and them executing it.

#rnx

diff --git a/Engine/Plugins/Enterprise/DatasmithImporter/Source/DatasmithImporter/Private/DatasmithConsumer.cpp b/Engine/Plugins/Enterprise/DatasmithImporter/Source/DatasmithImporter/Private/DatasmithConsumer.cpp
--- a/Engine/Plugins/Enterprise/DatasmithImporter/Source/DatasmithImporter/Private/DatasmithConsumer.cpp
+++ b/Engine/Plugins/Enterprise/DatasmithImporter/Source/DatasmithImporter/Private/DatasmithConsumer.cpp
@@ -31,11 +31,11 @@
 
 #include "Algo/Count.h"
 #include "AssetRegistryModule.h"
+#include "CineCameraActor.h"
+#include "CineCameraComponent.h"
 #include "ComponentReregisterContext.h"
 #include "Components/HierarchicalInstancedStaticMeshComponent.h"
 #include "Components/PointLightComponent.h"
-#include "CineCameraActor.h"
-#include "CineCameraComponent.h"
 #include "EditorLevelUtils.h"
 #include "Engine/Level.h"
 #include "Engine/LevelStreamingDynamic.h"
@@ -58,6 +58,7 @@
 #include "Misc/SecureHash.h"
 #include "UObject/Package.h"
 #include "UObject/SoftObjectPtr.h"
+#include "UObject/UnrealType.h"
 
 #define LOCTEXT_NAMESPACE "DatasmithConsumer"
 
@@ -88,7 +89,7 @@ bool UDatasmithConsumer::Initialize()
 
 	ProgressTaskPtr->ReportNextStep( LOCTEXT( "DatasmithImportFactory_Initialize", "Preparing world ...") );
 
-	MoveAssets();
+	UpdateDestinationPackage();
 
 	MoveLevel();
 
@@ -194,8 +195,7 @@ bool UDatasmithConsumer::Run()
 	ProgressTaskPtr->ReportNextStep( LOCTEXT( "DatasmithImportFactory_Finalize", "Finalizing commit ...") );
 	FDatasmithImporter::FinalizeImport( *ImportContextPtr, TSet<UObject*>() );
 
-	// Store package path and level name for subsequent call to Run
-	LastPackagePath = TargetContentFolder;
+	// Store the level name for subsequent call to Run
 	LastLevelName = LevelName;
 
 	return true;
@@ -372,11 +372,11 @@ bool UDatasmithConsumer::SetLevelName( const FString & InLevelName, FText& OutRe
 	return bValidLevelName;
 }
 
-void UDatasmithConsumer::MoveAssets()
+void UDatasmithConsumer::UpdateDestinationPackage()
 {
 	// Do nothing if this is the First call to Run, DatasmithScene is null and LastPackagePath is empty
 	// or the re-Run is using the same package path
-	if( ( !DatasmithScene.IsValid() && LastPackagePath.IsEmpty() ) || LastPackagePath == TargetContentFolder )
+	if( !DatasmithScene.IsValid() )
 	{
 		return;
 	}
@@ -384,9 +384,9 @@ void UDatasmithConsumer::MoveAssets()
 	const FText DialogTitle( LOCTEXT( "DatasmithConsumerDlgTitle", "Warning" ) );
 
 	// Warn user if related Datasmith scene is not in package path and force re-creation of Datasmith scene
-	if( DatasmithScene.IsValid() && !WorkingScenePtr->GetPathName().StartsWith( TargetContentFolder ) )
+	if( DatasmithScene.IsValid() && !DatasmithScene->GetPathName().StartsWith( TargetContentFolder ) )
 	{
-		FText WarningMessage = FText::Format(LOCTEXT("DatasmithConsumer_NoSceneAsset", "Package path {0} different from path previously used, {1}.\nPrevious content will not be updated."), FText::FromString (TargetContentFolder ), FText::FromString ( LastPackagePath ) );
+		FText WarningMessage = FText::Format(LOCTEXT("DatasmithConsumer_NoSceneAsset", "Package path {0} different from path previously used, {1}.\nPrevious content will not be updated."), FText::FromString (TargetContentFolder ), FText::FromString ( DatasmithScene->GetOutermost()->GetPathName() ) );
 		FMessageDialog::Open(EAppMsgType::Ok, WarningMessage, &DialogTitle );
 
 		UE_LOG( LogDatasmithImport, Warning, TEXT("%s"), *WarningMessage.ToString() );
diff --git a/Engine/Plugins/Enterprise/DatasmithImporter/Source/DatasmithImporter/Private/DatasmithConsumer.h b/Engine/Plugins/Enterprise/DatasmithImporter/Source/DatasmithImporter/Private/DatasmithConsumer.h
--- a/Engine/Plugins/Enterprise/DatasmithImporter/Source/DatasmithImporter/Private/DatasmithConsumer.h
+++ b/Engine/Plugins/Enterprise/DatasmithImporter/Source/DatasmithImporter/Private/DatasmithConsumer.h
@@ -28,13 +28,9 @@ public:
 	{
 	}
 
-	UPROPERTY( BlueprintReadOnly, Category = DatasmithConsumerInternal)
+	UPROPERTY( BlueprintReadOnly, Category = DatasmithConsumerInternal, DuplicateTransient )
 	TSoftObjectPtr<UDatasmithScene> DatasmithScene;
 
-	/** Stores the package path used on the last call to UDatasmithConsumer::Run */
-	UPROPERTY( BlueprintReadOnly, Category = DatasmithConsumerInternal )
-	FString LastPackagePath;
-
 	/** Stores the level used on the last call to UDatasmithConsumer::Run */
 	UPROPERTY( BlueprintReadOnly, Category = DatasmithConsumerInternal )
 	FString LastLevelName;
@@ -58,7 +54,7 @@ private:
 	ULevel* FindLevel( const FString& InLevelName );
 
 	/** Move assets if destination package path has changed since last call to UDatasmithConsumer::Run */
-	void MoveAssets();
+	void UpdateDestinationPackage();
 
 	/** Move level if destination level's name has changed since last call to UDatasmithConsumer::Run */
 	void MoveLevel();
