Added a missing if not nullptr.

diff --git a/Engine/Plugins/Enterprise/DatasmithImporter/Source/DatasmithImporter/Private/DatasmithImportFactory.cpp b/Engine/Plugins/Enterprise/DatasmithImporter/Source/DatasmithImporter/Private/DatasmithImportFactory.cpp
--- a/Engine/Plugins/Enterprise/DatasmithImporter/Source/DatasmithImporter/Private/DatasmithImportFactory.cpp
+++ b/Engine/Plugins/Enterprise/DatasmithImporter/Source/DatasmithImporter/Private/DatasmithImportFactory.cpp
@@ -50,14 +50,17 @@ namespace DatasmithImportFactoryImpl
 				// This factory handles both UDatasmithSceneImportData and UDatasmithTranslatedSceneImportData but not other children of UDatasmithSceneImportData
 				UDatasmithSceneImportData* SceneAssetImportData = Cast< UDatasmithScene >(Obj)->AssetImportData;
 
-				if ( SceneAssetImportData->GetClass() == UDatasmithSceneImportData::StaticClass() )
+				if ( SceneAssetImportData )
 				{
-					return SceneAssetImportData;
-				}
-				else
-				{
-					// UDatasmithTranslatedSceneImportData are associated with scenes imported through Translators system
-					return ExactCast<UDatasmithTranslatedSceneImportData>(SceneAssetImportData);
+					if ( SceneAssetImportData->GetClass() == UDatasmithSceneImportData::StaticClass() )
+					{
+						return SceneAssetImportData;
+					}
+					else
+					{
+						// UDatasmithTranslatedSceneImportData are associated with scenes imported through Translators system
+						return ExactCast<UDatasmithTranslatedSceneImportData>(SceneAssetImportData);
+					}
 				}
 			}
 			if (Class == UDatasmithSceneImportData::StaticClass())
