Improved meaning of name of scene generated by consumer

diff --git a/Engine/Plugins/Enterprise/DatasmithImporter/Source/DatasmithImporter/Private/DatasmithConsumer.cpp b/Engine/Plugins/Enterprise/DatasmithImporter/Source/DatasmithImporter/Private/DatasmithConsumer.cpp
--- a/Engine/Plugins/Enterprise/DatasmithImporter/Source/DatasmithImporter/Private/DatasmithConsumer.cpp
+++ b/Engine/Plugins/Enterprise/DatasmithImporter/Source/DatasmithImporter/Private/DatasmithConsumer.cpp
@@ -65,6 +65,8 @@
 const FText DatasmithConsumerLabel( LOCTEXT( "DatasmithConsumerLabel", "Datasmith writer" ) );
 const FText DatasmithConsumerDescription( LOCTEXT( "DatasmithConsumerDesc", "Writes data prep world's current level and assets to current level" ) );
 
+const TCHAR* DatasmithSceneSuffix = TEXT("_Scene");
+
 namespace DatasmithConsumerUtils
 {
 	/** Helper to generate actor element from a scene actor */
@@ -89,7 +91,7 @@ bool UDatasmithConsumer::Initialize()
 
 	ProgressTaskPtr->ReportNextStep( LOCTEXT( "DatasmithImportFactory_Initialize", "Preparing world ...") );
 
-	UpdateDestinationPackage();
+	UpdateScene();
 
 	MoveLevel();
 
@@ -100,15 +102,55 @@ bool UDatasmithConsumer::Initialize()
 
 	WorkingScenePtr.Reset( DatasmithScene.Get() );
 
-	// Check if the Datasmith scene is not already in memory
+	// Re-create the DatasmithScene if it is invalid
 	if ( !WorkingScenePtr.IsValid() )
 	{
-		FName DatasmithSceneName = MakeUniqueObjectName( ParentPackage, UDatasmithScene::StaticClass(), *( GetName() + TEXT("_DS") ) );
 
-		UPackage* Package = CreatePackage( nullptr, *FPaths::Combine( ParentPackage->GetPathName(), DatasmithSceneName.ToString() ) );
+		FString DatasmithSceneName = GetOuter()->GetName() + DatasmithSceneSuffix;
+
+		UPackage* Package = CreatePackage( nullptr, *FPaths::Combine( ParentPackage->GetPathName(), DatasmithSceneName ) );
 		Package->FullyLoad();
 
-		DatasmithScene = NewObject< UDatasmithScene >( Package, DatasmithSceneName, GetFlags() | RF_Public | RF_Transactional );
+		if( UObject* ExistingObject = StaticFindObject( nullptr, Package, *DatasmithSceneName, true ) )
+		{
+			bool bDatasmithSceneFound = false;
+
+			// Check to see if existing scene is not from same Dataprep asset
+			if( UDatasmithScene* ExistingDatasmithScene = Cast<UDatasmithScene>( ExistingObject ) )
+			{
+				if ( ExistingDatasmithScene->GetClass()->ImplementsInterface( UInterface_AssetUserData::StaticClass() ) )
+				{
+					if ( IInterface_AssetUserData* AssetUserDataInterface = Cast< IInterface_AssetUserData >( ExistingDatasmithScene ) )
+					{
+						if( UDataprepAssetUserData* DataprepAssetUserData = AssetUserDataInterface->GetAssetUserData< UDataprepAssetUserData >() )
+						{
+							UDataprepAssetInterface* DataprepAssetInterface = Cast< UDataprepAssetInterface >( GetOuter() );
+							check( DataprepAssetInterface );
+
+							if( DataprepAssetUserData->DataprepAssetPtr == DataprepAssetInterface )
+							{
+								DatasmithScene.Reset();
+								DatasmithScene = ExistingDatasmithScene;
+								Package = nullptr;
+								bDatasmithSceneFound = true;
+							}
+						}
+					}
+				}
+			}
+
+			if( !bDatasmithSceneFound )
+			{
+				DatasmithSceneName = MakeUniqueObjectName( ParentPackage, UDatasmithScene::StaticClass(), *DatasmithSceneName ).ToString();
+				Package = CreatePackage( nullptr, *FPaths::Combine( ParentPackage->GetPathName(), DatasmithSceneName ) );
+				Package->FullyLoad();
+			}
+		}
+
+		if(Package != nullptr)
+		{
+			DatasmithScene = NewObject< UDatasmithScene >( Package, *DatasmithSceneName, GetFlags() | RF_Public | RF_Transactional );
+		}
 		check( DatasmithScene.IsValid() );
 
 		WorkingScenePtr.Reset( DatasmithScene.Get() );
@@ -120,9 +162,6 @@ bool UDatasmithConsumer::Initialize()
 		check( WorkingScenePtr->AssetImportData );
 
 		// Store a Dataprep asset pointer into the scene asset in order to be able to later re-execute the dataprep pipeline
-		UDataprepAssetInterface* DataprepAssetInterface = Cast< UDataprepAssetInterface >( GetOuter() );
-		check( DataprepAssetInterface );
-
 		if ( WorkingScenePtr->GetClass()->ImplementsInterface(UInterface_AssetUserData::StaticClass()) )
 		{
 			if ( IInterface_AssetUserData* AssetUserDataInterface = Cast< IInterface_AssetUserData >( DatasmithScene.Get() ) )
@@ -135,6 +174,10 @@ bool UDatasmithConsumer::Initialize()
 					DataprepAssetUserData = NewObject< UDataprepAssetUserData >( DatasmithScene.Get(), NAME_None, Flags );
 					AssetUserDataInterface->AddAssetUserData( DataprepAssetUserData );
 				}
+
+				UDataprepAssetInterface* DataprepAssetInterface = Cast< UDataprepAssetInterface >( GetOuter() );
+				check( DataprepAssetInterface );
+
 				DataprepAssetUserData->DataprepAssetPtr = DataprepAssetInterface;
 			}
 		}
@@ -372,10 +415,9 @@ bool UDatasmithConsumer::SetLevelName( const FString & InLevelName, FText& OutRe
 	return bValidLevelName;
 }
 
-void UDatasmithConsumer::UpdateDestinationPackage()
+void UDatasmithConsumer::UpdateScene()
 {
-	// Do nothing if this is the First call to Run, DatasmithScene is null and LastPackagePath is empty
-	// or the re-Run is using the same package path
+	// Do nothing if this is the First call to Run, DatasmithScene is null
 	if( !DatasmithScene.IsValid() )
 	{
 		return;
@@ -384,7 +426,7 @@ void UDatasmithConsumer::UpdateDestinationPackage()
 	const FText DialogTitle( LOCTEXT( "DatasmithConsumerDlgTitle", "Warning" ) );
 
 	// Warn user if related Datasmith scene is not in package path and force re-creation of Datasmith scene
-	if( DatasmithScene.IsValid() && !DatasmithScene->GetPathName().StartsWith( TargetContentFolder ) )
+	if( !DatasmithScene->GetPathName().StartsWith( TargetContentFolder ) )
 	{
 		FText WarningMessage = FText::Format(LOCTEXT("DatasmithConsumer_NoSceneAsset", "Package path {0} different from path previously used, {1}.\nPrevious content will not be updated."), FText::FromString (TargetContentFolder ), FText::FromString ( DatasmithScene->GetOutermost()->GetPathName() ) );
 		FMessageDialog::Open(EAppMsgType::Ok, WarningMessage, &DialogTitle );
@@ -394,6 +436,18 @@ void UDatasmithConsumer::UpdateDestinationPackage()
 		// Force re-creation of Datasmith scene
 		DatasmithScene.Reset();
 	}
+	// Check if name of owning Dataprep asset has not changed
+	else
+	{
+		const FString DatasmithSceneName = GetOuter()->GetName() + DatasmithSceneSuffix;
+
+		if( DatasmithScene->GetName() != DatasmithSceneName )
+		{
+			// Force re-creation of Datasmith scene
+			DatasmithScene.Reset();
+		}
+	}
+
 }
 
 void UDatasmithConsumer::MoveLevel()
diff --git a/Engine/Plugins/Enterprise/DatasmithImporter/Source/DatasmithImporter/Private/DatasmithConsumer.h b/Engine/Plugins/Enterprise/DatasmithImporter/Source/DatasmithImporter/Private/DatasmithConsumer.h
--- a/Engine/Plugins/Enterprise/DatasmithImporter/Source/DatasmithImporter/Private/DatasmithConsumer.h
+++ b/Engine/Plugins/Enterprise/DatasmithImporter/Source/DatasmithImporter/Private/DatasmithConsumer.h
@@ -54,7 +54,7 @@ private:
 	ULevel* FindLevel( const FString& InLevelName );
 
 	/** Move assets if destination package path has changed since last call to UDatasmithConsumer::Run */
-	void UpdateDestinationPackage();
+	void UpdateScene();
 
 	/** Move level if destination level's name has changed since last call to UDatasmithConsumer::Run */
 	void MoveLevel();
