Fixed issue with path when opening import dialog

#rnx

diff --git a/Engine/Plugins/Enterprise/DatasmithContent/Source/DatasmithContentEditor/Private/DatasmithSceneDetails.cpp b/Engine/Plugins/Enterprise/DatasmithContent/Source/DatasmithContentEditor/Private/DatasmithSceneDetails.cpp
--- a/Engine/Plugins/Enterprise/DatasmithContent/Source/DatasmithContentEditor/Private/DatasmithSceneDetails.cpp
+++ b/Engine/Plugins/Enterprise/DatasmithContent/Source/DatasmithContentEditor/Private/DatasmithSceneDetails.cpp
@@ -203,7 +203,7 @@ FReply FDatasmithSceneDetails::OnSelectFile()
 	if ( bOpened && OpenedFiles.Num() > 0 )
 	{
 		FString& OpenedFile = OpenedFiles[0];
-		FEditorDirectories::Get().SetLastDirectory( ELastDirectory::GENERIC_IMPORT, OpenedFile );
+		FEditorDirectories::Get().SetLastDirectory( ELastDirectory::GENERIC_IMPORT, FPaths::GetPath(OpenedFile) );
 
 		FString Extension = FPaths::GetExtension( OpenedFile );
 
diff --git a/Engine/Plugins/Enterprise/DatasmithImporter/Source/DatasmithImporter/Private/DatasmithFileProducer.cpp b/Engine/Plugins/Enterprise/DatasmithImporter/Source/DatasmithImporter/Private/DatasmithFileProducer.cpp
--- a/Engine/Plugins/Enterprise/DatasmithImporter/Source/DatasmithImporter/Private/DatasmithFileProducer.cpp
+++ b/Engine/Plugins/Enterprise/DatasmithImporter/Source/DatasmithImporter/Private/DatasmithFileProducer.cpp
@@ -1370,7 +1370,7 @@ FString FDatasmithFileProducerUtils::SelectFileToImport()
 	if ( bOpened && OpenedFiles.Num() > 0 )
 	{
 		const FString& OpenedFile = OpenedFiles[0];
-		FEditorDirectories::Get().SetLastDirectory( ELastDirectory::GENERIC_IMPORT, OpenedFile );
+		FEditorDirectories::Get().SetLastDirectory( ELastDirectory::GENERIC_IMPORT, FPaths::GetPath(OpenedFile) );
 
 		return FPaths::ConvertRelativePathToFull( OpenedFile );
 	}
