Fixed issue with non-persisted DatasmithScene

diff --git a/Engine/Plugins/Enterprise/DatasmithImporter/Source/DatasmithImporter/Private/DatasmithConsumer.cpp b/Engine/Plugins/Enterprise/DatasmithImporter/Source/DatasmithImporter/Private/DatasmithConsumer.cpp
--- a/Engine/Plugins/Enterprise/DatasmithImporter/Source/DatasmithImporter/Private/DatasmithConsumer.cpp
+++ b/Engine/Plugins/Enterprise/DatasmithImporter/Source/DatasmithImporter/Private/DatasmithConsumer.cpp
@@ -100,10 +100,8 @@ bool UDatasmithConsumer::Initialize()
 	UPackage* ParentPackage = CreatePackage( nullptr, *TargetContentFolder );
 	ParentPackage->FullyLoad();
 
-	WorkingScenePtr.Reset( DatasmithScene.Get() );
-
 	// Re-create the DatasmithScene if it is invalid
-	if ( !WorkingScenePtr.IsValid() )
+	if ( !DatasmithScene.IsValid() )
 	{
 
 		FString DatasmithSceneName = GetOuter()->GetName() + DatasmithSceneSuffix;
@@ -149,20 +147,19 @@ bool UDatasmithConsumer::Initialize()
 
 		if(Package != nullptr)
 		{
-			DatasmithScene = NewObject< UDatasmithScene >( Package, *DatasmithSceneName, GetFlags() | RF_Public | RF_Transactional );
+			DatasmithScene = NewObject< UDatasmithScene >( Package, *DatasmithSceneName, GetFlags() | RF_Standalone | RF_Public | RF_Transactional );
 		}
 		check( DatasmithScene.IsValid() );
 
-		WorkingScenePtr.Reset( DatasmithScene.Get() );
+		DatasmithScene->MarkPackageDirty();
 
 		FAssetRegistryModule::AssetCreated( DatasmithScene.Get() );
-		WorkingScenePtr->MarkPackageDirty();
 
-		WorkingScenePtr->AssetImportData = NewObject< UDatasmithSceneImportData >( DatasmithScene.Get(), UDatasmithSceneImportData::StaticClass() );
-		check( WorkingScenePtr->AssetImportData );
+		DatasmithScene->AssetImportData = NewObject< UDatasmithSceneImportData >( DatasmithScene.Get(), UDatasmithSceneImportData::StaticClass() );
+		check( DatasmithScene->AssetImportData );
 
 		// Store a Dataprep asset pointer into the scene asset in order to be able to later re-execute the dataprep pipeline
-		if ( WorkingScenePtr->GetClass()->ImplementsInterface(UInterface_AssetUserData::StaticClass()) )
+		if ( DatasmithScene->GetClass()->ImplementsInterface(UInterface_AssetUserData::StaticClass()) )
 		{
 			if ( IInterface_AssetUserData* AssetUserDataInterface = Cast< IInterface_AssetUserData >( DatasmithScene.Get() ) )
 			{
@@ -249,7 +246,6 @@ void UDatasmithConsumer::Reset()
 	ImportContextPtr.Reset();
 	ProgressTaskPtr.Reset();
 	UDataprepContentConsumer::Reset();
-	WorkingScenePtr.Reset(nullptr);
 
 	// Restore previous current level
 	if( PreviousCurrentLevel != nullptr )
@@ -271,7 +267,7 @@ const FText& UDatasmithConsumer::GetDescription() const
 
 bool UDatasmithConsumer::BuildContexts( UWorld* ImportWorld )
 {
-	const FString FilePath = FPaths::Combine( FPaths::ProjectIntermediateDir(), ( WorkingScenePtr->GetName() + TEXT( ".udatasmith" ) ) );
+	const FString FilePath = FPaths::Combine( FPaths::ProjectIntermediateDir(), ( DatasmithScene->GetName() + TEXT( ".udatasmith" ) ) );
 
 	ImportContextPtr = MakeUnique< FDatasmithImportContext >( FilePath, false, TEXT("DatasmithImport"), LOCTEXT("DatasmithImportFactoryDescription", "Datasmith") );
 
@@ -279,7 +275,7 @@ bool UDatasmithConsumer::BuildContexts( UWorld* ImportWorld )
 	ImportContextPtr->Options->BaseOptions.SceneHandling = EDatasmithImportScene::CurrentLevel;
 	ImportContextPtr->SceneAsset = DatasmithScene.Get();
 	ImportContextPtr->ActorsContext.ImportWorld = ImportWorld;
-	ImportContextPtr->Scene = FDatasmithSceneFactory::CreateScene( *WorkingScenePtr->GetName() );
+	ImportContextPtr->Scene = FDatasmithSceneFactory::CreateScene( *DatasmithScene->GetName() );
 	ImportContextPtr->SceneName = ImportContextPtr->Scene->GetName();
 
 	// Convert all incoming Datasmith scene actors as regular actors
@@ -293,7 +289,7 @@ bool UDatasmithConsumer::BuildContexts( UWorld* ImportWorld )
 	FDatasmithImporterUtils::FillSceneElement( ImportContextPtr->Scene, RootActors );
 
 	// Initialize context
-	FString SceneOuterPath = WorkingScenePtr->GetOutermost()->GetName();
+	FString SceneOuterPath = DatasmithScene->GetOutermost()->GetName();
 	FString RootPath = FPackageName::GetLongPackagePath( SceneOuterPath );
 
 	if ( Algo::Count( RootPath, TEXT('/') ) > 1 )
diff --git a/Engine/Plugins/Enterprise/DatasmithImporter/Source/DatasmithImporter/Private/DatasmithConsumer.h b/Engine/Plugins/Enterprise/DatasmithImporter/Source/DatasmithImporter/Private/DatasmithConsumer.h
--- a/Engine/Plugins/Enterprise/DatasmithImporter/Source/DatasmithImporter/Private/DatasmithConsumer.h
+++ b/Engine/Plugins/Enterprise/DatasmithImporter/Source/DatasmithImporter/Private/DatasmithConsumer.h
@@ -65,7 +65,6 @@ private:
 private:
 	TUniquePtr< FDatasmithImportContext > ImportContextPtr;
 	TUniquePtr< FDataprepWorkReporter > ProgressTaskPtr;
-	TStrongObjectPtr< UDatasmithScene > WorkingScenePtr;
 
 	ULevel* PreviousCurrentLevel;
 };
\ No newline at end of file
