Fixed crash when loading solidworks file: Root cause - null assets can be generated by translators

diff --git a/Engine/Plugins/Enterprise/DatasmithImporter/Source/DatasmithImporter/Private/DatasmithFileProducer.cpp b/Engine/Plugins/Enterprise/DatasmithImporter/Source/DatasmithImporter/Private/DatasmithFileProducer.cpp
--- a/Engine/Plugins/Enterprise/DatasmithImporter/Source/DatasmithImporter/Private/DatasmithFileProducer.cpp
+++ b/Engine/Plugins/Enterprise/DatasmithImporter/Source/DatasmithImporter/Private/DatasmithFileProducer.cpp
@@ -279,41 +279,58 @@ void UDatasmithFileProducer::SceneElementToWorld()
 	{
 		TRACE_CPUPROFILER_EVENT_SCOPE(CommitMeshDescriptions);
 
+		UStaticMesh::FCommitMeshDescriptionParams Params;
+		Params.bMarkPackageDirty = false;
+		Params.bUseHashAsGuid = true;
+
 		ParallelFor(StaticMeshes.Num(),
 			[&](int32 StaticMeshIndex)
 			{
-				UStaticMesh::FCommitMeshDescriptionParams Params;
-				Params.bMarkPackageDirty = false;
-				Params.bUseHashAsGuid = true;
-
-				for (int32 Index = 0; Index < StaticMeshes[StaticMeshIndex]->GetNumSourceModels(); ++Index)
+				if(UStaticMesh* StaticMesh = StaticMeshes[StaticMeshIndex])
 				{
-					StaticMeshes[StaticMeshIndex]->CommitMeshDescription( Index, Params );
+					for (int32 Index = 0; Index < StaticMesh->GetNumSourceModels(); ++Index)
+					{
+						StaticMesh->CommitMeshDescription( Index, Params );
+					}
 				}
 			}
 		);
 	}
 
-	Assets.Append( StaticMeshes );
+	// Note: Some of the assets might be null (incomplete or failed import), only add non-null ones to Assets
+
+	for(UStaticMesh* StaticMesh : StaticMeshes)
+	{
+		if(StaticMesh)
+		{
+			Assets.Emplace( StaticMesh );
+		}
+	}
 
 	for ( TPair< TSharedRef< IDatasmithTextureElement >, UTexture* >& AssetPair : ImportContextPtr->ImportedTextures )
 	{
-		Assets.Emplace( AssetPair.Value );
+		if(AssetPair.Value)
+		{
+			Assets.Emplace( AssetPair.Value );
+		}
 	}
 
 	for ( TPair< TSharedRef< IDatasmithBaseMaterialElement >, UMaterialInterface* >& AssetPair : ImportContextPtr->ImportedMaterials )
 	{
-		Assets.Emplace( AssetPair.Value );
-
-		if (UMaterial* SourceMaterial = Cast< UMaterial >(AssetPair.Value))
+		if(AssetPair.Value)
 		{
-			SourceMaterial->RebuildExpressionTextureReferences();
+			Assets.Emplace( AssetPair.Value );
 
-			for (FMaterialFunctionInfo& MaterialFunctionInfo : SourceMaterial->MaterialFunctionInfos)
+			if (UMaterial* SourceMaterial = Cast< UMaterial >(AssetPair.Value))
 			{
-				if (MaterialFunctionInfo.Function && MaterialFunctionInfo.Function->GetOutermost() == SourceMaterial->GetOutermost())
+				SourceMaterial->RebuildExpressionTextureReferences();
+
+				for (FMaterialFunctionInfo& MaterialFunctionInfo : SourceMaterial->MaterialFunctionInfos)
 				{
-					Assets.Emplace( MaterialFunctionInfo.Function );
+					if (MaterialFunctionInfo.Function && MaterialFunctionInfo.Function->GetOutermost() == SourceMaterial->GetOutermost())
+					{
+						Assets.Emplace( MaterialFunctionInfo.Function );
+					}
 				}
 			}
 		}
@@ -321,22 +338,34 @@ void UDatasmithFileProducer::SceneElementToWorld()
 
 	for ( TPair< int32, UMaterialInterface* >& AssetPair : ImportContextPtr->ImportedParentMaterials )
 	{
-		Assets.Emplace( AssetPair.Value );
+		if(AssetPair.Value)
+		{
+			Assets.Emplace( AssetPair.Value );
+		}
 	}
 
 	for ( TPair< TSharedRef< IDatasmithBaseMaterialElement >, UMaterialFunction* >& AssetPair : ImportContextPtr->ImportedMaterialFunctions )
 	{
-		Assets.Emplace( AssetPair.Value );
+		if(AssetPair.Value)
+		{
+			Assets.Emplace( AssetPair.Value );
+		}
 	}
 
 	for ( TPair< TSharedRef< IDatasmithLevelSequenceElement >, ULevelSequence* >& AssetPair : ImportContextPtr->ImportedLevelSequences )
 	{
-		Assets.Emplace( AssetPair.Value );
+		if(AssetPair.Value)
+		{
+			Assets.Emplace( AssetPair.Value );
+		}
 	}
 
 	for ( TPair< TSharedRef< IDatasmithLevelVariantSetsElement >, ULevelVariantSets* >& AssetPair : ImportContextPtr->ImportedLevelVariantSets )
 	{
-		Assets.Emplace( AssetPair.Value );
+		if(AssetPair.Value)
+		{
+			Assets.Emplace( AssetPair.Value );
+		}
 	}
 }
 
