Revisited creation of assets in directory producer

diff --git a/Engine/Plugins/Enterprise/DatasmithImporter/Source/DatasmithImporter/Private/DatasmithFileProducer.cpp b/Engine/Plugins/Enterprise/DatasmithImporter/Source/DatasmithImporter/Private/DatasmithFileProducer.cpp
--- a/Engine/Plugins/Enterprise/DatasmithImporter/Source/DatasmithImporter/Private/DatasmithFileProducer.cpp
+++ b/Engine/Plugins/Enterprise/DatasmithImporter/Source/DatasmithImporter/Private/DatasmithFileProducer.cpp
@@ -739,11 +739,6 @@ bool UDatasmithDirProducer::Initialize()
 		return false;
 	}
 
-	UPackage* TransientPackage = NewObject< UPackage >( nullptr, *FPaths::Combine( Context.RootPackagePtr->GetPathName(), *GetName() ), RF_Transient );
-	TransientPackage->FullyLoad();
-
-	Context.SetRootPackage( TransientPackage );
-
 	FileProducer = TStrongObjectPtr< UDatasmithFileProducer >( NewObject< UDatasmithFileProducer >( GetTransientPackage(), NAME_None, RF_Transient ) );
 
 	return true;
@@ -759,6 +754,13 @@ bool UDatasmithDirProducer::Execute(TArray< TWeakObjectPtr< UObject > >& OutAsse
 
 	FDataprepWorkReporter Task( Context.ProgressReporterPtr, LOCTEXT( "DatasmithFileProducer_LoadingFromDirectory", "Loading files from directory ..." ), (float)FilesToProcess.Num(), 1.0f );
 
+	// Cache context's package
+	UPackage* CachedPackage = Context.RootPackagePtr.Get();
+
+	FString RootPath = FPaths::Combine( Context.RootPackagePtr->GetPathName(), *GetName() );
+	UPackage* RootTransientPackage = NewObject< UPackage >( nullptr, *RootPath, RF_Transient );
+	RootTransientPackage->FullyLoad();
+
 	for( const FString& FileName : FilesToProcess )
 	{
 		if ( IsCancelled() )
@@ -766,7 +768,21 @@ bool UDatasmithDirProducer::Execute(TArray< TWeakObjectPtr< UObject > >& OutAsse
 			break;
 		}
 
+		// Import content of file into the proper content folder to avoid name collision
+		UPackage* TransientPackage = RootTransientPackage;
+		FString FilePath = FPaths::GetPath( FileName );
+		if(FilePath != FolderPath)
+		{
+			FString SubFolder = FilePath.Right(FilePath.Len() - FolderPath.Len() - 1 /* Remove leading '/' */);
+			TransientPackage = NewObject< UPackage >( nullptr, *FPaths::Combine( RootPath, SubFolder ), RF_Transient );
+			TransientPackage->FullyLoad();
+		}
+
+		Context.SetRootPackage( TransientPackage );
+
+		// Update file producer's filename
 		FileProducer->FilePath =  FPaths::ConvertRelativePathToFull( FileName );
+		FileProducer->UpdateName();
 
 		Task.ReportNextStep( FText::Format( LOCTEXT( "DatasmithFileProducer_LoadingFile", "Loading {0} ..."), FText::FromString( FileName ) ) );
 
@@ -777,6 +793,9 @@ bool UDatasmithDirProducer::Execute(TArray< TWeakObjectPtr< UObject > >& OutAsse
 		}
 	}
 
+	// Restore context's package
+	Context.SetRootPackage( CachedPackage );
+
 	return !IsCancelled();
 }
 
diff --git a/Engine/Plugins/Enterprise/DatasmithImporter/Source/DatasmithImporter/Private/Utility/DatasmithImporterUtils.cpp b/Engine/Plugins/Enterprise/DatasmithImporter/Source/DatasmithImporter/Private/Utility/DatasmithImporterUtils.cpp
--- a/Engine/Plugins/Enterprise/DatasmithImporter/Source/DatasmithImporter/Private/Utility/DatasmithImporterUtils.cpp
+++ b/Engine/Plugins/Enterprise/DatasmithImporter/Source/DatasmithImporter/Private/Utility/DatasmithImporterUtils.cpp
@@ -311,10 +311,51 @@ void FDatasmithImporterUtils::DeleteActor( AActor& Actor )
 			}
 		}
 
+		// Clean up all references to external assets within the actor and its components since those will be deleted later
+		// #ueent_remark: Underlying question why the actor and its components are still reachable after being 'deleted'
+		{
+			struct FObjectExternalReferenceCleaner : public FArchiveUObject
+			{
+				FObjectExternalReferenceCleaner()	{}
+
+				virtual FArchive& operator<<(UObject*& ObjRef) override
+				{
+					if(ObjRef != nullptr)
+					{
+						// Set to null any pointer to an external asset
+						if( ObjRef->HasAnyFlags( RF_Standalone | RF_Public ) )
+						{
+							ObjRef = nullptr;
+						}
+					}
+
+					return *this;
+				}
+			};
+
+			TArray< UObject* > SubObjectsArray;
+			GetObjectsWithOuter( &Actor, SubObjectsArray, /*bIncludeNestedObjects = */ true );
+
+			for(UObject* SubObject : SubObjectsArray)
+			{
+				if(SubObject)
+				{
+					FObjectExternalReferenceCleaner Ar;
+					SubObject->Serialize(Ar);
+				}
+			}
+
+			{
+				FObjectExternalReferenceCleaner Ar;
+				Actor.Serialize(Ar);
+			}
+		}
+
+		// Actually delete the actor
 		ActorWorld->EditorDestroyActor( &Actor, true );
 
-		// Move the actor outside it's so it's object name can be use
-		Actor.UObject::Rename( nullptr, GetTransientPackage(), REN_DontCreateRedirectors | REN_ForceNoResetLoaders );
+		// Move the actor to the transient package so its object name can be use
+		Actor.UObject::Rename( *(Actor.GetName() + TEXT("_DELETED")), GetTransientPackage(), REN_DontCreateRedirectors | REN_ForceNoResetLoaders );
 	}
 }
 
