Fixed logic issue in previous submission

diff --git a/Engine/Plugins/Enterprise/DatasmithImporter/Source/DatasmithImporter/Private/Utility/DatasmithImporterUtils.cpp b/Engine/Plugins/Enterprise/DatasmithImporter/Source/DatasmithImporter/Private/Utility/DatasmithImporterUtils.cpp
--- a/Engine/Plugins/Enterprise/DatasmithImporter/Source/DatasmithImporter/Private/Utility/DatasmithImporterUtils.cpp
+++ b/Engine/Plugins/Enterprise/DatasmithImporter/Source/DatasmithImporter/Private/Utility/DatasmithImporterUtils.cpp
@@ -355,7 +355,7 @@ void FDatasmithImporterUtils::DeleteActor( AActor& Actor )
 		ActorWorld->EditorDestroyActor( &Actor, true );
 
 		// Move the actor to the transient package so its object name can be use
-		Actor.UObject::Rename( *(Actor.GetName() + TEXT("_DELETED")), GetTransientPackage(), REN_DontCreateRedirectors | REN_ForceNoResetLoaders );
+		Actor.UObject::Rename( nullptr, GetTransientPackage(), REN_DontCreateRedirectors | REN_ForceNoResetLoaders );
 	}
 }
 
