Fixed issue with renaming of content folder in Dataprep editor

diff --git a/Engine/Plugins/Enterprise/DatasmithImporter/Source/DatasmithImporter/Private/DatasmithConsumer.cpp b/Engine/Plugins/Enterprise/DatasmithImporter/Source/DatasmithImporter/Private/DatasmithConsumer.cpp
--- a/Engine/Plugins/Enterprise/DatasmithImporter/Source/DatasmithImporter/Private/DatasmithConsumer.cpp
+++ b/Engine/Plugins/Enterprise/DatasmithImporter/Source/DatasmithImporter/Private/DatasmithConsumer.cpp
@@ -422,7 +422,7 @@ void UDatasmithConsumer::UpdateScene()
 	const FText DialogTitle( LOCTEXT( "DatasmithConsumerDlgTitle", "Warning" ) );
 
 	// Warn user if related Datasmith scene is not in package path and force re-creation of Datasmith scene
-	if( !DatasmithScene->GetPathName().StartsWith( TargetContentFolder ) )
+	if( DatasmithScene->GetOutermost()->GetPathName() != TargetContentFolder )
 	{
 		FText WarningMessage = FText::Format(LOCTEXT("DatasmithConsumer_NoSceneAsset", "Package path {0} different from path previously used, {1}.\nPrevious content will not be updated."), FText::FromString (TargetContentFolder ), FText::FromString ( DatasmithScene->GetOutermost()->GetPathName() ) );
 		FMessageDialog::Open(EAppMsgType::Ok, WarningMessage, &DialogTitle );
diff --git a/Engine/Plugins/Enterprise/DatasmithImporter/Source/DatasmithImporter/Private/DatasmithImporter.cpp b/Engine/Plugins/Enterprise/DatasmithImporter/Source/DatasmithImporter/Private/DatasmithImporter.cpp
--- a/Engine/Plugins/Enterprise/DatasmithImporter/Source/DatasmithImporter/Private/DatasmithImporter.cpp
+++ b/Engine/Plugins/Enterprise/DatasmithImporter/Source/DatasmithImporter/Private/DatasmithImporter.cpp
@@ -584,10 +584,17 @@ namespace DatasmithImporterImpl
 		{
 			if ( ImportSceneActor )
 			{
-				for ( const auto& Pair : ImportSceneActor->RelatedActors )
+
+				TArray< TSoftObjectPtr< AActor > > RelatedActors;
+				ImportSceneActor->RelatedActors.GenerateValueArray( RelatedActors );
+
+				ImportSceneActor->Scene = nullptr;
+				ImportSceneActor->RelatedActors.Empty();
+
+				while(RelatedActors.Num() > 0)
 				{
-					AActor* RelatedActor = Pair.Value.Get();
-					if ( RelatedActor )
+					TSoftObjectPtr< AActor > ActorPtr = RelatedActors.Pop(false);
+					if(AActor* RelatedActor = ActorPtr.Get())
 					{
 						FDatasmithImporterUtils::DeleteActor( *RelatedActor );
 					}
