Fixed logic flaw in comparison of package paths

diff --git a/Engine/Plugins/Enterprise/DatasmithImporter/Source/DatasmithImporter/Private/DatasmithConsumer.cpp b/Engine/Plugins/Enterprise/DatasmithImporter/Source/DatasmithImporter/Private/DatasmithConsumer.cpp
--- a/Engine/Plugins/Enterprise/DatasmithImporter/Source/DatasmithImporter/Private/DatasmithConsumer.cpp
+++ b/Engine/Plugins/Enterprise/DatasmithImporter/Source/DatasmithImporter/Private/DatasmithConsumer.cpp
@@ -422,9 +422,10 @@ void UDatasmithConsumer::UpdateScene()
 	const FText DialogTitle( LOCTEXT( "DatasmithConsumerDlgTitle", "Warning" ) );
 
 	// Warn user if related Datasmith scene is not in package path and force re-creation of Datasmith scene
-	if( DatasmithScene->GetOutermost()->GetPathName() != TargetContentFolder )
+	FString DatasmithScenePath = FPaths::GetPath( DatasmithScene->GetPathName() );
+	if( DatasmithScenePath != TargetContentFolder )
 	{
-		FText WarningMessage = FText::Format(LOCTEXT("DatasmithConsumer_NoSceneAsset", "Package path {0} different from path previously used, {1}.\nPrevious content will not be updated."), FText::FromString (TargetContentFolder ), FText::FromString ( DatasmithScene->GetOutermost()->GetPathName() ) );
+		FText WarningMessage = FText::Format(LOCTEXT("DatasmithConsumer_NoSceneAsset", "Package path {0} different from path previously used, {1}.\nPrevious content will not be updated."), FText::FromString (TargetContentFolder ), FText::FromString ( DatasmithScenePath ) );
 		FMessageDialog::Open(EAppMsgType::Ok, WarningMessage, &DialogTitle );
 
 		UE_LOG( LogDatasmithImport, Warning, TEXT("%s"), *WarningMessage.ToString() );
diff --git a/Engine/Plugins/Experimental/DataPrepEditor/Source/DataPrepCore/Private/DataPrepContentConsumer.cpp b/Engine/Plugins/Experimental/DataPrepEditor/Source/DataPrepCore/Private/DataPrepContentConsumer.cpp
--- a/Engine/Plugins/Experimental/DataPrepEditor/Source/DataPrepCore/Private/DataPrepContentConsumer.cpp
+++ b/Engine/Plugins/Experimental/DataPrepEditor/Source/DataPrepCore/Private/DataPrepContentConsumer.cpp
@@ -81,6 +81,12 @@ bool UDataprepContentConsumer::SetTargetContentFolder(const FString& InTargetCon
 
 		TargetContentFolder = !InTargetContentFolder.IsEmpty() ? InTargetContentFolder : FPaths::GetPath( GetOutermost()->GetPathName() );
 
+		// Remove ending '/' if applicable
+		if(TargetContentFolder[TargetContentFolder.Len() - 1] == L'/')
+		{
+			TargetContentFolder.RemoveAt(TargetContentFolder.Len() - 1, 1);
+		}
+
 		OnChanged.Broadcast();
 	}
 
