Moved the SetLabel blueprint API from UDatasmithActorElement to UDatasmithObjectElement to match the C++ API

#rnx

diff --git a/Engine/Plugins/Enterprise/DatasmithImporter/Source/DatasmithImporter/Private/DatasmithImporter.cpp b/Engine/Plugins/Enterprise/DatasmithImporter/Source/DatasmithImporter/Private/DatasmithImporter.cpp
--- a/Engine/Plugins/Enterprise/DatasmithImporter/Source/DatasmithImporter/Private/DatasmithImporter.cpp
+++ b/Engine/Plugins/Enterprise/DatasmithImporter/Source/DatasmithImporter/Private/DatasmithImporter.cpp
@@ -2345,7 +2345,7 @@ void FDatasmithImporter::FinalizeImport(FDatasmithImportContext& ImportContext,
 
 	// Sometimes, the data is invalid and we get the same UStaticMesh multiple times
 	TSet< UStaticMesh* > StaticMeshes;
-	for (const TPair< TSharedRef< IDatasmithMeshElement >, UStaticMesh* >& ImportedStaticMeshPair : ImportContext.ImportedStaticMeshes)
+	for (TPair< TSharedRef< IDatasmithMeshElement >, UStaticMesh* >& ImportedStaticMeshPair : ImportContext.ImportedStaticMeshes)
 	{
 		if (ImportContext.bUserCancelled)
 		{
@@ -2372,6 +2372,7 @@ void FDatasmithImporter::FinalizeImport(FDatasmithImportContext& ImportContext,
 		ExistingStaticMeshPtr = FinalizeStaticMesh(SourceStaticMesh, *DestinationPackagePath, ExistingStaticMesh, &ReferencesToRemap, false);
 		DatasmithImporterImpl::CheckAssetPersistenceValidity(ExistingStaticMeshPtr.Get(), ImportContext);
 
+		ImportedStaticMeshPair.Value = ExistingStaticMeshPtr.Get();
 		StaticMeshes.Add(ExistingStaticMeshPtr.Get());
 	}
 
diff --git a/Engine/Plugins/Enterprise/DatasmithImporter/Source/DatasmithImporter/Private/ObjectElements/DatasmithUObjectElements.cpp b/Engine/Plugins/Enterprise/DatasmithImporter/Source/DatasmithImporter/Private/ObjectElements/DatasmithUObjectElements.cpp
--- a/Engine/Plugins/Enterprise/DatasmithImporter/Source/DatasmithImporter/Private/ObjectElements/DatasmithUObjectElements.cpp
+++ b/Engine/Plugins/Enterprise/DatasmithImporter/Source/DatasmithImporter/Private/ObjectElements/DatasmithUObjectElements.cpp
@@ -58,6 +58,18 @@ FString UDatasmithObjectElement::GetElementName() const
 	return FString(Element->GetName());
 }
 
+FString UDatasmithObjectElement::GetLabel() const
+{
+	DATASMITHOBJECTELEMENT_GETSHARED_AND_EARLYRETURN(IDatasmithElement, GetIDatasmithElement(), FString());
+	return Element->GetLabel();
+}
+
+void UDatasmithObjectElement::SetLabel(const FString& InLabel)
+{
+	DATASMITHOBJECTELEMENT_GETSHARED_AND_EARLYRETURN(IDatasmithElement, GetIDatasmithElement(), );
+	Element->SetLabel(*InLabel);
+}
+
 bool UDatasmithObjectElement::IsElementValid() const
 {
 	if (!GetIDatasmithElement().IsValid())
@@ -119,18 +131,6 @@ void UDatasmithActorElement::SetLayer(const FString& InLayer)
 	Element->SetLayer(*InLayer);
 }
 
-FString UDatasmithActorElement::GetLabel() const
-{
-	DATASMITHOBJECTELEMENT_GETSHARED_AND_EARLYRETURN(IDatasmithActorElement, GetIDatasmithActorElement(), FString());
-	return Element->GetLabel();
-}
-
-void UDatasmithActorElement::SetLabel(const FString& InLabel)
-{
-	DATASMITHOBJECTELEMENT_GETSHARED_AND_EARLYRETURN(IDatasmithActorElement, GetIDatasmithActorElement(), );
-	Element->SetLabel(*InLabel);
-}
-
 TArray<FString> UDatasmithActorElement::GetTags() const
 {
 	DATASMITHOBJECTELEMENT_GETSHARED_AND_EARLYRETURN(IDatasmithActorElement, GetIDatasmithActorElement(), TArray<FString>());
