Fixing the dependancy issues of the datasmith content plugin.
Also Fixed a ux issue allowing the user to create a dataprep asset even when the datasmith importer is off.

#rnx

diff --git a/Engine/Plugins/Enterprise/DatasmithContent/DatasmithContent.uplugin b/Engine/Plugins/Enterprise/DatasmithContent/DatasmithContent.uplugin
--- a/Engine/Plugins/Enterprise/DatasmithContent/DatasmithContent.uplugin
+++ b/Engine/Plugins/Enterprise/DatasmithContent/DatasmithContent.uplugin
@@ -30,14 +30,6 @@
 		{
 			"Name": "VariantManagerContent",
 			"Enabled": true
-		},
-		{
-			"Name": "DataprepEditor",
-			"Enabled": true
-		},
-		{
-			"Name": "EditorScriptingUtilities",
-			"Enabled": true
 		}
 	]
 }
\ No newline at end of file
diff --git a/Engine/Plugins/Enterprise/DatasmithContent/Source/DatasmithContentEditor/DatasmithContentEditor.Build.cs b/Engine/Plugins/Enterprise/DatasmithContent/Source/DatasmithContentEditor/DatasmithContentEditor.Build.cs
--- a/Engine/Plugins/Enterprise/DatasmithContent/Source/DatasmithContentEditor/DatasmithContentEditor.Build.cs
+++ b/Engine/Plugins/Enterprise/DatasmithContent/Source/DatasmithContentEditor/DatasmithContentEditor.Build.cs
@@ -13,9 +13,7 @@ namespace UnrealBuildTool.Rules
 					"Core",
 					"CoreUObject",
 					"DatasmithContent",
-					"DataprepCore",
 					"DesktopPlatform",
-					"EditorScriptingUtilities",
 					"EditorStyle",
 					"Engine",
 					"Projects",
diff --git a/Engine/Plugins/Enterprise/DatasmithContent/Source/DatasmithContentEditor/Private/AssetTypeActions_DatasmithScene.cpp b/Engine/Plugins/Enterprise/DatasmithContent/Source/DatasmithContentEditor/Private/AssetTypeActions_DatasmithScene.cpp
--- a/Engine/Plugins/Enterprise/DatasmithContent/Source/DatasmithContentEditor/Private/AssetTypeActions_DatasmithScene.cpp
+++ b/Engine/Plugins/Enterprise/DatasmithContent/Source/DatasmithContentEditor/Private/AssetTypeActions_DatasmithScene.cpp
@@ -4,8 +4,7 @@
 
 #include "DatasmithAssetImportData.h"
 #include "DatasmithContentEditorModule.h"
-#include "DataprepAssetUserData.h"
-#include "DataprepCoreLibrary.h"
+#include "DatasmithScene.h"
 
 #include "Framework/MultiBox/MultiBoxBuilder.h"
 
@@ -18,7 +17,7 @@ uint32 FAssetTypeActions_DatasmithScene::GetCategories()
 
 FText FAssetTypeActions_DatasmithScene::GetName() const
 {
-	return NSLOCTEXT("AssetTypeActions_DatasmithScene", "AssetTypeActions_DatasmithScene_Name", "Datasmith Scene");
+	return LOCTEXT("AssetTypeActions_DatasmithScene_Name", "Datasmith Scene");
 }
 
 UClass* FAssetTypeActions_DatasmithScene::GetSupportedClass() const
@@ -64,46 +63,4 @@ void FAssetTypeActions_DatasmithScene::OpenAssetEditor(const TArray<UObject*>& I
 	}
 }
 
-void FAssetTypeActions_DatasmithScene::ExecuteDataprepAssets( TArray< TSoftObjectPtr< UDataprepAssetInterface > > DataprepAssets )
-{
-	for ( TSoftObjectPtr< UDataprepAssetInterface > DataprepAssetInterface : DataprepAssets ) 
-	{
-		UDataprepCoreLibrary::ExecuteWithReporting( DataprepAssetInterface.Get() );
-	}
-}
-
-void FAssetTypeActions_DatasmithScene::GetActions( const TArray<UObject*>& InObjects, FMenuBuilder& MenuBuilder )
-{
-	auto DatasmithScenes = GetTypedWeakObjectPtrs<UDatasmithScene>( InObjects );
-
-	if ( DatasmithScenes.Num() == 0 )
-	{
-		return;
-	}
-
-	// Allow execute only if we have at least one valid dataprep asset pointer
-	TArray< TSoftObjectPtr< UDataprepAssetInterface > > DataprepAssets;
-	for ( TWeakObjectPtr< UDatasmithScene > Scene : DatasmithScenes )
-	{
-		UDataprepAssetUserData* DataprepAssetUserData = CastChecked< UDataprepAssetUserData >( Scene->GetAssetUserDataOfClass( UDataprepAssetUserData::StaticClass() ), ECastCheckedType::NullAllowed );
-		if (DataprepAssetUserData && DataprepAssetUserData->DataprepAssetPtr)
-		{
-			DataprepAssets.Add( DataprepAssetUserData->DataprepAssetPtr );
-		}
-	}
-
-	if ( DataprepAssets.Num() > 0 )
-	{
-		MenuBuilder.AddMenuEntry(
-			LOCTEXT("RunAsset", "Execute"),
-			LOCTEXT("RunAssetTooltip", "Runs the Dataprep asset's producers, execute its recipe, finally runs the consumer"),
-			FSlateIcon(),
-			FUIAction(
-				FExecuteAction::CreateSP( this, &FAssetTypeActions_DatasmithScene::ExecuteDataprepAssets, DataprepAssets ),
-				FCanExecuteAction()
-			)
-		);
-	}
-}
-
 #undef LOCTEXT_NAMESPACE
diff --git a/Engine/Plugins/Enterprise/DatasmithContent/Source/DatasmithContentEditor/Private/AssetTypeActions_DatasmithScene.h b/Engine/Plugins/Enterprise/DatasmithContent/Source/DatasmithContentEditor/Private/AssetTypeActions_DatasmithScene.h
--- a/Engine/Plugins/Enterprise/DatasmithContent/Source/DatasmithContentEditor/Private/AssetTypeActions_DatasmithScene.h
+++ b/Engine/Plugins/Enterprise/DatasmithContent/Source/DatasmithContentEditor/Private/AssetTypeActions_DatasmithScene.h
@@ -3,8 +3,6 @@
 #pragma once
 
 #include "CoreMinimal.h"
-#include "DatasmithScene.h"
-#include "DataprepAssetInterface.h"
 #include "AssetTypeActions_Base.h"
 
 /** Asset type actions for UDatasmithScene class */
@@ -16,13 +14,9 @@ public:
 	virtual uint32 GetCategories() override;
 	virtual FColor GetTypeColor() const override { return FColor(255, 0, 0); }
 	virtual UClass* GetSupportedClass() const override;
-	virtual bool HasActions ( const TArray<UObject*>& InObjects ) const override { return true; }
-	virtual void GetActions( const TArray<UObject*>& InObjects, FMenuBuilder& MenuBuilder ) override;
+	virtual bool HasActions ( const TArray<UObject*>& InObjects ) const override { return false; }
 	virtual bool IsImportedAsset() const override { return true; }
 	virtual void GetResolvedSourceFilePaths(const TArray<UObject*>& TypeAssets, TArray<FString>& OutSourceFilePaths) const override;
 	virtual void OpenAssetEditor(const TArray<UObject*>& InObjects, TSharedPtr<class IToolkitHost> EditWithinLevelEditor = TSharedPtr<IToolkitHost>()) override;
 	// End IAssetTypeActions interface
-
-private:
-	void ExecuteDataprepAssets(TArray< TSoftObjectPtr< UDataprepAssetInterface > > DataprepAssets);
 };
diff --git a/Engine/Plugins/Enterprise/DatasmithImporter/Source/DatasmithImporter/DatasmithImporter.Build.cs b/Engine/Plugins/Enterprise/DatasmithImporter/Source/DatasmithImporter/DatasmithImporter.Build.cs
--- a/Engine/Plugins/Enterprise/DatasmithImporter/Source/DatasmithImporter/DatasmithImporter.Build.cs
+++ b/Engine/Plugins/Enterprise/DatasmithImporter/Source/DatasmithImporter/DatasmithImporter.Build.cs
@@ -35,16 +35,17 @@ namespace UnrealBuildTool.Rules
 					"MessageLog",
 					"MovieScene",
 					"MovieSceneTracks",
-                    "RawMesh",
 					"RHI",
+					"RawMesh",
 					"Slate",
 					"SlateCore",
 					"SourceControl",
 					"StaticMeshDescription",
+					"StaticMeshEditorExtension",
+					"ToolMenus",
 					"UnrealEd",
 					"VariantManager",
 					"VariantManagerContent",
-                    "StaticMeshEditorExtension",
                 }
 			);
 
diff --git a/Engine/Plugins/Enterprise/DatasmithImporter/Source/DatasmithImporter/Private/DatasmithConsumer.cpp b/Engine/Plugins/Enterprise/DatasmithImporter/Source/DatasmithImporter/Private/DatasmithConsumer.cpp
--- a/Engine/Plugins/Enterprise/DatasmithImporter/Source/DatasmithImporter/Private/DatasmithConsumer.cpp
+++ b/Engine/Plugins/Enterprise/DatasmithImporter/Source/DatasmithImporter/Private/DatasmithConsumer.cpp
@@ -411,6 +411,17 @@ bool UDatasmithConsumer::SetLevelName( const FString & InLevelName, FText& OutRe
 	return bValidLevelName;
 }
 
+bool UDatasmithConsumer::SetTargetContentFolder(const FString& InTargetContentFolder, FText& OutReason)
+{
+	if ( Super::SetTargetContentFolder( InTargetContentFolder, OutReason ) )
+	{
+		UpdateScene();
+		return true;
+	}
+
+	return false;
+}
+
 void UDatasmithConsumer::UpdateScene()
 {
 	// Do nothing if this is the First call to Run, DatasmithScene is null
diff --git a/Engine/Plugins/Enterprise/DatasmithImporter/Source/DatasmithImporter/Private/DatasmithConsumer.h b/Engine/Plugins/Enterprise/DatasmithImporter/Source/DatasmithImporter/Private/DatasmithConsumer.h
--- a/Engine/Plugins/Enterprise/DatasmithImporter/Source/DatasmithImporter/Private/DatasmithConsumer.h
+++ b/Engine/Plugins/Enterprise/DatasmithImporter/Source/DatasmithImporter/Private/DatasmithConsumer.h
@@ -39,6 +39,7 @@ public:
 	virtual const FText& GetLabel() const override;
 	virtual const FText& GetDescription() const override;
 	virtual bool SetLevelName(const FString& InLevelName, FText& OutReason ) override;
+	virtual bool SetTargetContentFolder(const FString& InTargetContentFolder, FText& OutReason) override;
 
 protected:
 	virtual bool Initialize() override;
diff --git a/Engine/Plugins/Enterprise/DatasmithImporter/Source/DatasmithImporter/Private/DatasmithImporterModule.cpp b/Engine/Plugins/Enterprise/DatasmithImporter/Source/DatasmithImporter/Private/DatasmithImporterModule.cpp
--- a/Engine/Plugins/Enterprise/DatasmithImporter/Source/DatasmithImporter/Private/DatasmithImporterModule.cpp
+++ b/Engine/Plugins/Enterprise/DatasmithImporter/Source/DatasmithImporter/Private/DatasmithImporterModule.cpp
@@ -5,13 +5,14 @@
 #include "ActorFactoryDatasmithScene.h"
 #include "DatasmithAssetImportData.h"
 #include "DatasmithAssetUserData.h"
+#include "DatasmithConsumer.h"
 #include "DatasmithContentEditorModule.h"
 #include "DatasmithContentEditorStyle.h"
 #include "DatasmithCustomAction.h"
 #include "DatasmithFileProducer.h"
+#include "DatasmithImportFactory.h"
 #include "DatasmithImporterEditorSettings.h"
 #include "DatasmithImporterHelper.h"
-#include "DatasmithImportFactory.h"
 #include "DatasmithScene.h"
 #include "DatasmithStaticMeshImporter.h"
 #include "DatasmithUtils.h"
@@ -26,7 +27,11 @@
 
 #include "AssetToolsModule.h"
 #include "ContentBrowserDelegates.h"
+#include "ContentBrowserMenuContexts.h"
 #include "ContentBrowserModule.h"
+#include "DataprepAssetInterface.h"
+#include "DataprepAssetUserData.h"
+#include "DataprepCoreLibrary.h"
 #include "Editor.h"
 #include "EditorFramework/AssetImportData.h"
 #include "EditorStyleSet.h"
@@ -43,6 +48,11 @@
 #include "Modules/ModuleManager.h"
 #include "PropertyEditorModule.h"
 #include "Settings/EditorLoadingSavingSettings.h"
+#include "ToolMenu.h"
+#include "ToolMenuDelegates.h"
+#include "ToolMenuSection.h"
+#include "ToolMenus.h"
+#include "UObject/StrongObjectPtr.h"
 
 #define LOCTEXT_NAMESPACE "DatasmithImporter"
 
@@ -82,6 +92,8 @@ public:
 			FPropertyEditorModule& PropertyModule = FModuleManager::LoadModuleChecked< FPropertyEditorModule >( TEXT("PropertyEditor") );
 			PropertyModule.RegisterCustomClassLayout( TEXT("DatasmithFileProducer"), FOnGetDetailCustomizationInstance::CreateStatic( &FDatasmithFileProducerDetails::MakeDetails ) );
 			PropertyModule.RegisterCustomClassLayout( TEXT("DatasmithDirProducer"), FOnGetDetailCustomizationInstance::CreateStatic( &FDatasmithDirProducerDetails::MakeDetails ) );
+
+			AddDataprepMenuEntryForDatasmithSceneAsset();
 		}
 	}
 
@@ -90,6 +102,8 @@ public:
 		// Disable any UI feature if running in command mode
 		if (!IsRunningCommandlet())
 		{
+			RemoveDataprepMenuEntryForDatasmithSceneAsset();
+
 			if ( SpawnSceneActorsDelegateHandle.IsValid() && FModuleManager::Get().IsModuleLoaded( TEXT("DatasmithContentEditor") ) )
 			{
 				IDatasmithContentEditorModule& DatasmithContentEditorModule = FModuleManager::GetModuleChecked< IDatasmithContentEditorModule >( TEXT("DatasmithContentEditor") );
@@ -143,6 +157,10 @@ private:
 	void SetupMenuEntry();
 	void OnClickedMenuEntry();
 
+	// Add the menu entry for a datasmith asset generated from a dataprep asset
+	void AddDataprepMenuEntryForDatasmithSceneAsset();
+	void RemoveDataprepMenuEntryForDatasmithSceneAsset();
+
 	void SetupContentBrowserContextMenuExtender();
 	void RemoveContentBrowserContextMenuExtender();
 
@@ -180,6 +198,92 @@ void FDatasmithImporterModule::OnClickedMenuEntry()
 	}
 }
 
+void FDatasmithImporterModule::AddDataprepMenuEntryForDatasmithSceneAsset()
+{
+	if (UToolMenu* Menu = UToolMenus::Get()->ExtendMenu(TEXT("ContentBrowser.AssetContextMenu.DatasmithScene")))
+	{
+
+		FNewToolMenuDelegate DataprepSectionConstructor;
+		DataprepSectionConstructor.BindLambda( [](UToolMenu* ToolMenu)
+			{
+				if ( ToolMenu )
+				{
+					TArray<TStrongObjectPtr<UDataprepAssetInterface>> DataprepAssetInterfacesPtr;
+					if (UContentBrowserAssetContextMenuContext* ContentBrowserMenuContext = ToolMenu->FindContext<UContentBrowserAssetContextMenuContext>())
+					{
+						if (ContentBrowserMenuContext->CommonClass == UDatasmithScene::StaticClass())
+						{
+							TArray<UObject*> SelectedObjects = ContentBrowserMenuContext->GetSelectedObjects();
+							DataprepAssetInterfacesPtr.Reserve( SelectedObjects.Num() );
+							for (UObject* SelectedObject : SelectedObjects)
+							{
+								if (UDatasmithScene* SelectedDatasmithScene = Cast<UDatasmithScene>(SelectedObject))
+								{
+									if (UDataprepAssetUserData* DataprepAssetUserData = SelectedDatasmithScene->GetAssetUserData<UDataprepAssetUserData>())
+									{
+										if (UDataprepAssetInterface* DataprepAsset = DataprepAssetUserData->DataprepAssetPtr.LoadSynchronous())
+										{
+											if (UDatasmithConsumer* DatasmithConsumer = Cast<UDatasmithConsumer>(DataprepAsset->GetConsumer()))
+											{
+												if (DatasmithConsumer->DatasmithScene.LoadSynchronous() == SelectedDatasmithScene)
+												{
+													// A Dataprep asset was found and it will regenerate this scene
+													DataprepAssetInterfacesPtr.Emplace( DataprepAsset );
+													continue;
+												}
+											}
+										}
+									}
+								}
+
+								// Couldn't find the dataprep asset for this scene
+								return;
+							}
+						}
+					}
+					else
+					{
+						return;
+					}
+
+
+					FToolUIAction UIAction;
+					UIAction.ExecuteAction.BindLambda( [DataprepAssetInterfacesPtr](const FToolMenuContext&)
+						{
+							for ( const TStrongObjectPtr<UDataprepAssetInterface>& DataprepAssetInterfacePtr : DataprepAssetInterfacesPtr )
+							{
+								UDataprepCoreLibrary::ExecuteWithReporting( DataprepAssetInterfacePtr.Get() );
+							}
+						});
+
+					FToolMenuInsert MenuInsert;
+					MenuInsert.Position = EToolMenuInsertType::First;
+					FToolMenuSection& Section = ToolMenu->AddSection( TEXT("Dataprep"), LOCTEXT("Dataprep", "Dataprep"), MenuInsert );
+
+					Section.AddMenuEntry(
+						TEXT("UpdateDataprepGeneratedAsset"),
+						LOCTEXT("UpdateDataprepGeneratedAsset", "Update Datasmith Scene(s)"),
+						LOCTEXT("UpdateDataprepGeneratedAssetTooltip", "Update the asset(s) by executing the Dataprep asset(s) that created it."),
+						FSlateIcon(),
+						UIAction
+					);
+
+				}
+			});
+
+
+		FToolMenuSection& Section = Menu->AddDynamicSection(TEXT("Dataprep"), DataprepSectionConstructor);
+	}
+}
+
+void FDatasmithImporterModule::RemoveDataprepMenuEntryForDatasmithSceneAsset()
+{
+	if (UToolMenu* Menu = UToolMenus::Get()->ExtendMenu(TEXT("ContentBrowser.AssetContextMenu.DatasmithScene")))
+	{
+		Menu->RemoveSection(TEXT("Dataprep"));
+	}
+}
+
 void FDatasmithImporterModule::SetupContentBrowserContextMenuExtender()
 {
 	if ( !IsRunningCommandlet() )
diff --git a/Engine/Plugins/Experimental/DataPrepEditor/Source/DataPrepCore/Private/DataPrepFactories.cpp b/Engine/Plugins/Experimental/DataPrepEditor/Source/DataPrepCore/Private/DataPrepFactories.cpp
--- a/Engine/Plugins/Experimental/DataPrepEditor/Source/DataPrepCore/Private/DataPrepFactories.cpp
+++ b/Engine/Plugins/Experimental/DataPrepEditor/Source/DataPrepCore/Private/DataPrepFactories.cpp
@@ -8,6 +8,7 @@
 
 #include "AssetRegistryModule.h"
 #include "AssetTypeCategories.h"
+#include "UObject/UObjectHash.h"
 #include "UObject/UObjectIterator.h"
 
 UDataprepAssetFactory::UDataprepAssetFactory()
@@ -19,6 +20,23 @@ UDataprepAssetFactory::UDataprepAssetFactory()
 	bEditorImport = false;
 }
 
+bool UDataprepAssetFactory::ShouldShowInNewMenu() const
+{
+	// If there is no consumer don't show this factory
+	TArray< UClass* > PotentialClasses;
+	GetDerivedClasses( UDataprepContentConsumer::StaticClass(), PotentialClasses, true );
+
+	for ( UClass* ChildClass : PotentialClasses )
+	{
+		if ( ChildClass && !ChildClass->HasAnyClassFlags( CLASS_CompiledFromBlueprint | CLASS_Deprecated | CLASS_NewerVersionExists | CLASS_Abstract ) && ChildClass->HasAllClassFlags( CLASS_Native ) )
+		{
+			return true;
+		}
+	}
+
+	return false;
+}
+
 UObject * UDataprepAssetFactory::FactoryCreateNew(UClass* InClass, UObject* InParent, FName InName, EObjectFlags Flags, UObject* Context, FFeedbackContext *Warn)
 {
 	check( InClass->IsChildOf( UDataprepAsset::StaticClass() ) );
diff --git a/Engine/Plugins/Experimental/DataPrepEditor/Source/DataPrepCore/Public/DataPrepFactories.h b/Engine/Plugins/Experimental/DataPrepEditor/Source/DataPrepCore/Public/DataPrepFactories.h
--- a/Engine/Plugins/Experimental/DataPrepEditor/Source/DataPrepCore/Public/DataPrepFactories.h
+++ b/Engine/Plugins/Experimental/DataPrepEditor/Source/DataPrepCore/Public/DataPrepFactories.h
@@ -21,6 +21,7 @@ public:
 	UDataprepAssetFactory();
 
 	//~ Begin UFactory Interface
+	virtual bool ShouldShowInNewMenu() const override;
 	virtual UObject* FactoryCreateNew(UClass* InClass, UObject* InParent, FName InName, EObjectFlags Flags, UObject* Context, FFeedbackContext* Warn) override;
 	//~ Begin UFactory Interface	
 };
diff --git a/Engine/Plugins/Experimental/DataPrepEditor/Source/DataPrepCore/Public/DataprepCoreLibrary.h b/Engine/Plugins/Experimental/DataPrepEditor/Source/DataPrepCore/Public/DataprepCoreLibrary.h
--- a/Engine/Plugins/Experimental/DataPrepEditor/Source/DataPrepCore/Public/DataprepCoreLibrary.h
+++ b/Engine/Plugins/Experimental/DataPrepEditor/Source/DataPrepCore/Public/DataprepCoreLibrary.h
@@ -2,7 +2,6 @@
 
 #pragma once
 
-#include "EditorFilterLibrary.h"
 #include "Kismet/BlueprintFunctionLibrary.h"
 
 #include "DataprepCoreLibrary.generated.h"
