Fixed a crash when closing the editor from my last cl.

#rnx

diff --git a/Engine/Plugins/Enterprise/DatasmithImporter/Source/DatasmithImporter/Private/DatasmithImporterModule.cpp b/Engine/Plugins/Enterprise/DatasmithImporter/Source/DatasmithImporter/Private/DatasmithImporterModule.cpp
--- a/Engine/Plugins/Enterprise/DatasmithImporter/Source/DatasmithImporter/Private/DatasmithImporterModule.cpp
+++ b/Engine/Plugins/Enterprise/DatasmithImporter/Source/DatasmithImporter/Private/DatasmithImporterModule.cpp
@@ -278,9 +278,12 @@ void FDatasmithImporterModule::AddDataprepMenuEntryForDatasmithSceneAsset()
 
 void FDatasmithImporterModule::RemoveDataprepMenuEntryForDatasmithSceneAsset()
 {
-	if (UToolMenu* Menu = UToolMenus::Get()->ExtendMenu(TEXT("ContentBrowser.AssetContextMenu.DatasmithScene")))
+	if (UToolMenus* Singlethon = UToolMenus::Get())
 	{
-		Menu->RemoveSection(TEXT("Dataprep"));
+		if (UToolMenu* Menu = Singlethon->ExtendMenu(TEXT("ContentBrowser.AssetContextMenu.DatasmithScene")))
+		{
+			Menu->RemoveSection(TEXT("Dataprep"));
+		}
 	}
 }
 
