Fixed issue with package path not ending with a '/'

diff --git a/Engine/Plugins/Enterprise/DatasmithImporter/Source/DatasmithImporter/Private/DatasmithConsumer.cpp b/Engine/Plugins/Enterprise/DatasmithImporter/Source/DatasmithImporter/Private/DatasmithConsumer.cpp
--- a/Engine/Plugins/Enterprise/DatasmithImporter/Source/DatasmithImporter/Private/DatasmithConsumer.cpp
+++ b/Engine/Plugins/Enterprise/DatasmithImporter/Source/DatasmithImporter/Private/DatasmithConsumer.cpp
@@ -97,7 +97,7 @@ bool UDatasmithConsumer::Initialize()
 
 	UpdateLevel();
 
-	UPackage* ParentPackage = CreatePackage( nullptr, *TargetContentFolder );
+	UPackage* ParentPackage = CreatePackage( nullptr, *GetTargetPackagePath() );
 	ParentPackage->FullyLoad();
 
 	// Re-create the DatasmithScene if it is invalid
@@ -436,13 +436,13 @@ void UDatasmithConsumer::UpdateScene()
 	FString DatasmithScenePath = FPaths::GetPath( DatasmithScene->GetPathName() );
 	if( DatasmithScenePath != TargetContentFolder )
 	{
+		// Force re-creation of Datasmith scene
+		DatasmithScene.Reset();
+
 		FText WarningMessage = FText::Format(LOCTEXT("DatasmithConsumer_NoSceneAsset", "Package path {0} different from path previously used, {1}.\nPrevious content will not be updated."), FText::FromString (TargetContentFolder ), FText::FromString ( DatasmithScenePath ) );
 		FMessageDialog::Open(EAppMsgType::Ok, WarningMessage, &DialogTitle );
 
 		UE_LOG( LogDatasmithImport, Warning, TEXT("%s"), *WarningMessage.ToString() );
-
-		// Force re-creation of Datasmith scene
-		DatasmithScene.Reset();
 	}
 	// Check if name of owning Dataprep asset has not changed
 	else
diff --git a/Engine/Plugins/Experimental/DataPrepEditor/Source/DataPrepCore/Private/DataPrepContentConsumer.cpp b/Engine/Plugins/Experimental/DataPrepEditor/Source/DataPrepCore/Private/DataPrepContentConsumer.cpp
--- a/Engine/Plugins/Experimental/DataPrepEditor/Source/DataPrepCore/Private/DataPrepContentConsumer.cpp
+++ b/Engine/Plugins/Experimental/DataPrepEditor/Source/DataPrepCore/Private/DataPrepContentConsumer.cpp
@@ -93,6 +93,31 @@ bool UDataprepContentConsumer::SetTargetContentFolder(const FString& InTargetCon
 	return bValidContentFolder;
 }
 
+FString UDataprepContentConsumer::GetTargetPackagePath() const
+{
+	FString TargetPackagePath(TargetContentFolder);
+
+	if( TargetPackagePath.IsEmpty() )
+	{
+		TargetPackagePath = TEXT("/Game/");
+	}
+	else if( TargetPackagePath.StartsWith( TEXT("/Content") ) )
+	{
+		TargetPackagePath = TargetPackagePath.Replace( TEXT("/Content"), TEXT("/Game") );
+	}
+
+	// If path is one level deep, make sure it ends with a '/'
+	int32 Index = -1;
+	TargetPackagePath.FindLastChar(L'/', Index);
+	check(Index >= 0);
+	if(Index == 0)
+	{
+		TargetPackagePath.Append( TEXT("/") );
+	}
+
+	return TargetPackagePath;
+}
+
 bool UDataprepContentConsumer::SetLevelName(const FString & InLevelName, FText& OutReason)
 {
 	OutReason = LOCTEXT( "DataprepContentConsumer_SetLevelName", "Not implemented" );
diff --git a/Engine/Plugins/Experimental/DataPrepEditor/Source/DataPrepCore/Public/DataPrepContentConsumer.h b/Engine/Plugins/Experimental/DataPrepEditor/Source/DataPrepCore/Public/DataPrepContentConsumer.h
--- a/Engine/Plugins/Experimental/DataPrepEditor/Source/DataPrepCore/Public/DataPrepContentConsumer.h
+++ b/Engine/Plugins/Experimental/DataPrepEditor/Source/DataPrepCore/Public/DataPrepContentConsumer.h
@@ -108,6 +108,8 @@ public:
 	 */
 	virtual bool SetLevelName(const FString& InLevelName, FText& OutReason );
 
+	const FString& GetLevelName() { return LevelName; }
+
 	/**
 	 * Sets the path of the package the consumer should move assets to if applicable.
 	 * Generally, this package path is substituted to the temporary path the assets are in
@@ -117,10 +119,14 @@ public:
 	 */
 	virtual bool SetTargetContentFolder(const FString& InTargetContentFolder, FText& OutReason );
 
-	const FString& GetLevelName() { return LevelName; }
-
 	const FString& GetTargetContentFolder() { return TargetContentFolder; }
 
+	/**
+	 * Returns a well-formed path to use when calling CreatePackage to create the target package.
+	 * @remark 
+	 */
+	FString GetTargetPackagePath() const;
+
 	/**
 	 * Allow an observer to be notified when one of the properties of the consumer changes
 	 * @return The delegate that will be broadcasted when the consumer changed
