Force root components to movable if they are the target of an animation

diff --git a/Engine/Plugins/Enterprise/DatasmithImporter/Source/DatasmithImporter/Private/DatasmithLevelSequenceImporter.cpp b/Engine/Plugins/Enterprise/DatasmithImporter/Source/DatasmithImporter/Private/DatasmithLevelSequenceImporter.cpp
--- a/Engine/Plugins/Enterprise/DatasmithImporter/Source/DatasmithImporter/Private/DatasmithLevelSequenceImporter.cpp
+++ b/Engine/Plugins/Enterprise/DatasmithImporter/Source/DatasmithImporter/Private/DatasmithLevelSequenceImporter.cpp
@@ -372,6 +372,11 @@ namespace DatasmithLevelSequenceImporterImpl
 		}
 
 		TransformSection->SetMask((EMovieSceneTransformChannel)TransformAnimation->GetEnabledTransformChannels());
+
+		if (USceneComponent* SceneComp = Actor->GetRootComponent())
+		{
+			SceneComp->SetMobility(EComponentMobility::Movable);
+		}
 	}
 
 	void PopulateVisibilityTrack(const TSharedRef<IDatasmithVisibilityAnimationElement>& VisibilityAnimation, int32 AnimIndex, FDatasmithImportContext& ImportContext, ULevelSequence* LevelSequence, FFrameNumber& MinFrameNumber, FFrameNumber& MaxFrameNumber)
@@ -445,6 +450,11 @@ namespace DatasmithLevelSequenceImporterImpl
 		{
 			PropagateVisibilityTrackRecursive(Actor, AnimIndex + 1, LevelSequence, MovieScene, VisibilityTrack);
 		}
+
+		if (USceneComponent* SceneComp = Actor->GetRootComponent())
+		{
+			SceneComp->SetMobility(EComponentMobility::Movable);
+		}
 	}
 
 	void PopulateSubsequenceTrack(const TSharedRef<IDatasmithSubsequenceAnimationElement>& SubsequenceAnimation, bool bShouldClearMasterTrack, FDatasmithImportContext& ImportContext, ULevelSequence* LevelSequence, FFrameNumber& MinFrameNumber, FFrameNumber& MaxFrameNumber)
