Fixed an ensure that ocurred during some datasmith reimport.

#rnx

diff --git a/Engine/Plugins/Enterprise/DatasmithImporter/Source/DatasmithImporter/Private/Utility/DatasmithImporterUtils.cpp b/Engine/Plugins/Enterprise/DatasmithImporter/Source/DatasmithImporter/Private/Utility/DatasmithImporterUtils.cpp
--- a/Engine/Plugins/Enterprise/DatasmithImporter/Source/DatasmithImporter/Private/Utility/DatasmithImporterUtils.cpp
+++ b/Engine/Plugins/Enterprise/DatasmithImporter/Source/DatasmithImporter/Private/Utility/DatasmithImporterUtils.cpp
@@ -270,7 +270,11 @@ void FDatasmithImporterUtils::DeleteNonImportedDatasmithElementFromSceneActor(AD
 					// Removed the non imported components
 					for ( UActorComponent* ComponentToRemove : ComponentsToRemove )
 					{
-						ComponentToRemove->DestroyComponent( true );
+						// Some components can destroy other components when being destroyed
+						if ( !ComponentToRemove->IsBeingDestroyed() )
+						{
+							ComponentToRemove->DestroyComponent( true );
+						}
 					}
 				}
 				else
