Copy up //UE4/Dev-Enterprise/... to //UE4/Dev-Tools-Staging/...

diff --git a/Engine/Plugins/Enterprise/DatasmithImporter/Source/DatasmithImporter/Private/DatasmithImporter.cpp b/Engine/Plugins/Enterprise/DatasmithImporter/Source/DatasmithImporter/Private/DatasmithImporter.cpp
--- a/Engine/Plugins/Enterprise/DatasmithImporter/Source/DatasmithImporter/Private/DatasmithImporter.cpp
+++ b/Engine/Plugins/Enterprise/DatasmithImporter/Source/DatasmithImporter/Private/DatasmithImporter.cpp
@@ -36,8 +36,8 @@
 #include "Components/HierarchicalInstancedStaticMeshComponent.h"
 #include "Containers/Map.h"
 #include "CoreMinimal.h"
-#include "EditorLevelUtils.h"
 #include "Editor/UnrealEdEngine.h"
+#include "EditorLevelUtils.h"
 #include "Engine/Engine.h"
 #include "Engine/StaticMesh.h"
 #include "Engine/StaticMeshActor.h"
@@ -51,8 +51,8 @@
 #include "LevelSequence.h"
 #include "MaterialEditingLibrary.h"
 #include "MaterialShared.h"
-#include "Materials/MaterialFunction.h"
 #include "Materials/Material.h"
+#include "Materials/MaterialFunction.h"
 #include "Materials/MaterialInstanceConstant.h"
 #include "Misc/FeedbackContext.h"
 #include "Misc/FileHelper.h"
@@ -66,8 +66,9 @@
 #include "Serialization/ObjectWriter.h"
 #include "Settings/EditorExperimentalSettings.h"
 #include "SourceControlOperations.h"
-#include "UnrealEdGlobals.h"
+#include "Templates/UniquePtr.h"
 #include "UObject/Package.h"
+#include "UnrealEdGlobals.h"
 
 extern UNREALED_API UEditorEngine* GEditor;
 
@@ -75,6 +76,32 @@ extern UNREALED_API UEditorEngine* GEditor;
 
 namespace DatasmithImporterImpl
 {
+	static void ReportProgress(FScopedSlowTask* SlowTask, const float ExpectedWorkThisFrame, const FText& Text)
+	{
+		if (SlowTask)
+		{
+			SlowTask->EnterProgressFrame(ExpectedWorkThisFrame, Text);
+		}
+	}
+
+	static void ReportProgress(FScopedSlowTask* SlowTask, const float ExpectedWorkThisFrame, FText&& Text)
+	{
+		if (SlowTask)
+		{
+			SlowTask->EnterProgressFrame(ExpectedWorkThisFrame, Text);
+		}
+	}
+
+	static bool HasUserCancelledTask(FFeedbackContext* FeedbackContext)
+	{
+		if (FeedbackContext)
+		{
+			return FeedbackContext->ReceivedUserCancel();
+		}
+
+		return false;
+	}
+
 	static UObject* PublicizeAsset( UObject* SourceAsset, const TCHAR* DestinationPath, UObject* ExistingAsset )
 	{
 		TRACE_CPUPROFILER_EVENT_SCOPE(DatasmithImporterImpl::PublicizeAsset);
@@ -106,7 +133,7 @@ namespace DatasmithImporterImpl
 		{
 			OldAssetPathName = ExistingAsset->GetPathName();
 
-			DestinationAsset = DuplicateObject< UObject >( SourceAsset, DestinationPackage, ExistingAsset->GetFName() );
+			DestinationAsset = FDatasmithImporterUtils::DuplicateObject( SourceAsset, DestinationPackage, ExistingAsset->GetFName() );
 
 			// If mesh's label has changed, update its name
 			if ( ExistingAsset->GetFName() != SourceAsset->GetFName() )
@@ -269,9 +296,10 @@ namespace DatasmithImporterImpl
 		const int32 TexturesCount = ImportContext.FilteredScene->GetTexturesCount();
 		const int32 MaterialsCount = ImportContext.FilteredScene->GetMaterialsCount();
 
+		FFeedbackContext* FeedbackContext = ImportContext.FeedbackContext;
 		for ( int32 TextureIndex = 0; TextureIndex < TexturesCount && !ImportContext.bUserCancelled; ++TextureIndex )
 		{
-			ImportContext.bUserCancelled |= ImportContext.Warn->ReceivedUserCancel();
+			ImportContext.bUserCancelled |= DatasmithImporterImpl::HasUserCancelledTask( ImportContext.FeedbackContext );
 
 			TSharedPtr< IDatasmithTextureElement > TextureElement = ImportContext.FilteredScene->GetTexture( TextureIndex );
 			const FString TextureName = ObjectTools::SanitizeObjectName( TextureElement->GetName() );
@@ -874,8 +902,13 @@ void FDatasmithImporter::ImportStaticMeshes( FDatasmithImportContext& ImportCont
 
 	TRACE_CPUPROFILER_EVENT_SCOPE(FDatasmithImporter::ImportStaticMeshes);
 
-	FScopedSlowTask Progress(StaticMeshesCount, LOCTEXT("ImportStaticMeshes", "Importing Static Meshes..."), true, *ImportContext.Warn );
-	Progress.MakeDialog(true);
+	TUniquePtr<FScopedSlowTask> ProgressPtr;
+
+	if ( ImportContext.FeedbackContext )
+	{
+		ProgressPtr = MakeUnique<FScopedSlowTask>(StaticMeshesCount, LOCTEXT("ImportStaticMeshes", "Importing Static Meshes..."), true, *ImportContext.FeedbackContext );
+		ProgressPtr->MakeDialog(true);
+	}
 
 	TMap<TSharedRef<IDatasmithMeshElement>, TFuture<FDatasmithMeshElementPayload*>> MeshElementPayloads;
 
@@ -890,7 +923,8 @@ void FDatasmithImporter::ImportStaticMeshes( FDatasmithImportContext& ImportCont
 	{
 		for (int32 MeshIndex = 0; MeshIndex < StaticMeshesCount && !ImportContext.bUserCancelled; ++MeshIndex)
 		{
-			ImportContext.bUserCancelled |= ImportContext.Warn->ReceivedUserCancel();
+			ImportContext.bUserCancelled |= DatasmithImporterImpl::HasUserCancelledTask( ImportContext.FeedbackContext );
+
 
 			if (!ImportContext.AssetsContext.StaticMeshesFinalPackage || ImportContext.AssetsContext.StaticMeshesFinalPackage->GetFName() == NAME_None || ImportContext.SceneTranslator == nullptr)
 			{
@@ -926,14 +960,16 @@ void FDatasmithImporter::ImportStaticMeshes( FDatasmithImportContext& ImportCont
 		}
 	}
 
+	FScopedSlowTask* Progress = ProgressPtr.Get();
+
 	// This pass will wait on the futures we got from the first pass async tasks
 	for ( int32 MeshIndex = 0; MeshIndex < StaticMeshesCount && !ImportContext.bUserCancelled; ++MeshIndex )
 	{
-		ImportContext.bUserCancelled |= ImportContext.Warn->ReceivedUserCancel();
+		ImportContext.bUserCancelled |= DatasmithImporterImpl::HasUserCancelledTask( ImportContext.FeedbackContext );
 
 		TSharedRef< IDatasmithMeshElement > MeshElement = ImportContext.FilteredScene->GetMesh( MeshIndex ).ToSharedRef();
 
-		Progress.EnterProgressFrame( 1.f, FText::FromString( FString::Printf( TEXT("Importing static mesh %d/%d (%s) ..."), MeshIndex + 1, StaticMeshesCount, MeshElement->GetLabel() ) ) );
+		DatasmithImporterImpl::ReportProgress( Progress, 1.f, FText::FromString( FString::Printf( TEXT("Importing static mesh %d/%d (%s) ..."), MeshIndex + 1, StaticMeshesCount, MeshElement->GetLabel() ) ) );
 
 		UStaticMesh* ExistingStaticMesh = nullptr;
 
@@ -1086,8 +1122,13 @@ void FDatasmithImporter::ImportTextures( FDatasmithImportContext& ImportContext
 
 	const int32 TexturesCount = ImportContext.FilteredScene->GetTexturesCount();
 
-	FScopedSlowTask Progress( (float)TexturesCount, LOCTEXT("ImportingTextures", "Importing Textures..."), true, *ImportContext.Warn );
-	Progress.MakeDialog(true);
+	TUniquePtr<FScopedSlowTask> ProgressPtr;
+
+	if ( ImportContext.FeedbackContext )
+	{
+		ProgressPtr = MakeUnique<FScopedSlowTask>( (float)TexturesCount, LOCTEXT("ImportingTextures", "Importing Textures..."), true, *ImportContext.FeedbackContext );
+		ProgressPtr->MakeDialog(true);
+	}
 
 	if (ImportContext.Options->TextureConflictPolicy != EDatasmithImportAssetConflictPolicy::Ignore && TexturesCount > 0)
 	{
@@ -1119,7 +1160,7 @@ void FDatasmithImporter::ImportTextures( FDatasmithImportContext& ImportContext
 
 		for ( int32 TextureIndex = 0; TextureIndex < FilteredTextureElements.Num(); TextureIndex++ )
 		{
-			ImportContext.bUserCancelled |= ImportContext.Warn->ReceivedUserCancel();
+			ImportContext.bUserCancelled |= DatasmithImporterImpl::HasUserCancelledTask( ImportContext.FeedbackContext );
 
 			AsyncData[TextureIndex].Result = 
 				Async(
@@ -1136,9 +1177,14 @@ void FDatasmithImporter::ImportTextures( FDatasmithImportContext& ImportContext
 				);
 		}
 
+		// Avoid a call to IsValid for each item
+		FScopedSlowTask* Progress = ProgressPtr.Get();
+
 		for ( int32 TextureIndex = 0; TextureIndex < FilteredTextureElements.Num(); TextureIndex++ )
 		{
-			if ((ImportContext.bUserCancelled |= ImportContext.Warn->ReceivedUserCancel()))
+			ImportContext.bUserCancelled |= DatasmithImporterImpl::HasUserCancelledTask( ImportContext.FeedbackContext );
+
+			if ( ImportContext.bUserCancelled )
 			{
 				// If operation has been canceled, just wait for other threads to also cancel
 				AsyncData[TextureIndex].Result.Wait();
@@ -1147,7 +1193,7 @@ void FDatasmithImporter::ImportTextures( FDatasmithImportContext& ImportContext
 			{
 				const TSharedPtr< IDatasmithTextureElement >& TextureElement = FilteredTextureElements[TextureIndex];
 
-				Progress.EnterProgressFrame( 1.f, FText::FromString( FString::Printf( TEXT("Importing texture %d/%d (%s) ..."), TextureIndex + 1, FilteredTextureElements.Num(), TextureElement->GetLabel() ) ) );
+				DatasmithImporterImpl::ReportProgress( Progress, 1.f, FText::FromString( FString::Printf( TEXT("Importing texture %d/%d (%s) ..."), TextureIndex + 1, FilteredTextureElements.Num(), TextureElement->GetLabel() ) ) );
 
 				UTexture* ExistingTexture = nullptr;
 
@@ -1393,8 +1439,15 @@ void FDatasmithImporter::ImportActors( FDatasmithImportContext& ImportContext )
 
 	const int32 ActorsCount = ImportContext.Scene->GetActorsCount();
 
-	FScopedSlowTask Progress( (float)ActorsCount, LOCTEXT("ImportActors", "Spawning actors..."), true, *ImportContext.Warn );
-	Progress.MakeDialog(true);
+	TUniquePtr<FScopedSlowTask> ProgressPtr;
+
+	if ( ImportContext.FeedbackContext )
+	{
+		ProgressPtr = MakeUnique<FScopedSlowTask>( (float)ActorsCount, LOCTEXT("ImportActors", "Spawning actors..."), true, *ImportContext.FeedbackContext );
+		ProgressPtr->MakeDialog(true);
+	}
+
+	FScopedSlowTask* Progress = ProgressPtr.Get();
 
 	if ( ImportSceneActor )
 	{
@@ -1402,13 +1455,13 @@ void FDatasmithImporter::ImportActors( FDatasmithImportContext& ImportContext )
 
 		for (int32 i = 0; i < ActorsCount && !ImportContext.bUserCancelled; ++i)
 		{
-			ImportContext.bUserCancelled |= ImportContext.Warn->ReceivedUserCancel();
+			ImportContext.bUserCancelled |= DatasmithImporterImpl::HasUserCancelledTask( ImportContext.FeedbackContext );
 
 			TSharedPtr< IDatasmithActorElement > ActorElement = ImportContext.Scene->GetActor(i);
 
 			if ( ActorElement.IsValid() )
 			{
-				Progress.EnterProgressFrame( 1.f, FText::FromString( FString::Printf( TEXT("Spawning actor %d/%d (%s) ..."), ( i + 1 ), ActorsCount, ActorElement->GetLabel() ) ) );
+				DatasmithImporterImpl::ReportProgress( Progress, 1.f, FText::FromString( FString::Printf( TEXT("Spawning actor %d/%d (%s) ..."), ( i + 1 ), ActorsCount, ActorElement->GetLabel() ) ) );
 
 				if ( ActorElement->IsAComponent() )
 				{
@@ -1431,7 +1484,7 @@ void FDatasmithImporter::ImportActors( FDatasmithImportContext& ImportContext )
 		// After all actors were imported, perform a post import step so that any dependencies can be resolved
 		for (int32 i = 0; i < ActorsCount && !ImportContext.bUserCancelled; ++i)
 		{
-			ImportContext.bUserCancelled |= ImportContext.Warn->ReceivedUserCancel();
+			ImportContext.bUserCancelled |= DatasmithImporterImpl::HasUserCancelledTask( ImportContext.FeedbackContext );
 
 			TSharedPtr< IDatasmithActorElement > ActorElement = ImportContext.Scene->GetActor(i);
 
@@ -1515,7 +1568,7 @@ AActor* FDatasmithImporter::ImportActor( FDatasmithImportContext& ImportContext,
 
 	for (int32 i = 0; i < ActorElement->GetChildrenCount() && !ImportContext.bUserCancelled; ++i)
 	{
-		ImportContext.bUserCancelled |= ImportContext.Warn->ReceivedUserCancel();
+		ImportContext.bUserCancelled |= DatasmithImporterImpl::HasUserCancelledTask( ImportContext.FeedbackContext );
 
 		const TSharedPtr< IDatasmithActorElement >& ChildActorElement = ActorElement->GetChild(i);
 
@@ -1886,8 +1939,12 @@ void FDatasmithImporter::ImportLevelSequences( FDatasmithImportContext& ImportCo
 		return;
 	}
 
-	FScopedSlowTask Progress( (float)SequencesCount, LOCTEXT("ImportingLevelSequences", "Importing Level Sequences..."), true, *ImportContext.Warn );
-	Progress.MakeDialog(true);
+	TUniquePtr<FScopedSlowTask> ProgressPtr;
+	if ( ImportContext.FeedbackContext )
+	{ 
+		ProgressPtr = MakeUnique<FScopedSlowTask>( (float)SequencesCount, LOCTEXT("ImportingLevelSequences", "Importing Level Sequences..."), true, *ImportContext.FeedbackContext );
+		ProgressPtr->MakeDialog(true);
+	}
 
 	// We can only parse a IDatasmithLevelSequenceElement with IDatasmithSubsequenceAnimationElements if their target
 	// subsequences' LevelSequenceElement have been parsed. We solve that with a structure we can repeatedly loop over,
@@ -1896,7 +1953,7 @@ void FDatasmithImporter::ImportLevelSequences( FDatasmithImportContext& ImportCo
 	SequencesToImport.Reserve(SequencesCount);
 	for ( int32 SequenceIndex = 0; SequenceIndex < SequencesCount && !ImportContext.bUserCancelled; ++SequenceIndex )
 	{
-		ImportContext.bUserCancelled |= ImportContext.Warn->ReceivedUserCancel();
+		ImportContext.bUserCancelled |= DatasmithImporterImpl::HasUserCancelledTask( ImportContext.FeedbackContext );
 
 		TSharedPtr< IDatasmithLevelSequenceElement > SequenceElement = ImportContext.FilteredScene->GetLevelSequence( SequenceIndex );
 		if ( !SequenceElement )
@@ -1907,6 +1964,9 @@ void FDatasmithImporter::ImportLevelSequences( FDatasmithImportContext& ImportCo
 		SequencesToImport.Add(SequenceElement);
 	}
 
+
+	FScopedSlowTask* Progress = ProgressPtr.Get();
+
 	// If the scene is ok we will do at most HardLoopCounter passes
 	int32 HardLoopCounter = SequencesToImport.Num();
 	int32 NumImported = 0;
@@ -1916,7 +1976,7 @@ void FDatasmithImporter::ImportLevelSequences( FDatasmithImportContext& ImportCo
 		// Scan remaining sequences and import the ones we can, removing from this array
 		for ( int32 SequenceIndex = SequencesToImport.Num() - 1; SequenceIndex >= 0 && !ImportContext.bUserCancelled; --SequenceIndex )
 		{
-			ImportContext.bUserCancelled |= ImportContext.Warn->ReceivedUserCancel();
+			ImportContext.bUserCancelled |= DatasmithImporterImpl::HasUserCancelledTask( ImportContext.FeedbackContext );
 
 			TSharedPtr<IDatasmithLevelSequenceElement>& SequenceElement = SequencesToImport[SequenceIndex];
 
@@ -1936,8 +1996,8 @@ void FDatasmithImporter::ImportLevelSequences( FDatasmithImportContext& ImportCo
 			}
 		}
 
-			FString SequenceName = ObjectTools::SanitizeObjectName( SequenceElement->GetName() );
-			Progress.EnterProgressFrame( 1.f, FText::FromString( FString::Printf( TEXT("Importing level sequence %d/%d (%s) ..."), NumImported + 1, HardLoopCounter, *SequenceName ) ) );
+		FString SequenceName = ObjectTools::SanitizeObjectName(SequenceElement->GetName());
+		DatasmithImporterImpl::ReportProgress( Progress, 1.f, FText::FromString( FString::Printf(TEXT("Importing level sequence %d/%d (%s) ..."), NumImported + 1, HardLoopCounter, *SequenceName ) ) );
 
 		ULevelSequence*& ImportedLevelSequence = ImportContext.ImportedLevelSequences.FindOrAdd( SequenceElement.ToSharedRef() );
 		if (ImportContext.SceneTranslator)
@@ -1989,12 +2049,18 @@ void FDatasmithImporter::ImportLevelVariantSets( FDatasmithImportContext& Import
 
 	TRACE_CPUPROFILER_EVENT_SCOPE(FDatasmithImporter::ImportLevelVariantSets);
 
-	FScopedSlowTask Progress( (float)LevelVariantSetsCount, LOCTEXT("ImportingLevelVariantSets", "Importing Level Variant Sets..."), true, *ImportContext.Warn );
-	Progress.MakeDialog(true);
+	TUniquePtr<FScopedSlowTask> ProgressPtr;
+	if ( ImportContext.FeedbackContext )
+	{
+		ProgressPtr = MakeUnique<FScopedSlowTask>( (float)LevelVariantSetsCount, LOCTEXT("ImportingLevelVariantSets", "Importing Level Variant Sets..."), true, *ImportContext.FeedbackContext );
+		ProgressPtr->MakeDialog(true);
+	}
+
+	FScopedSlowTask* Progress = ProgressPtr.Get();
 
 	for ( int32 LevelVariantSetIndex = 0; LevelVariantSetIndex < LevelVariantSetsCount && !ImportContext.bUserCancelled; ++LevelVariantSetIndex )
 	{
-		ImportContext.bUserCancelled |= ImportContext.Warn->ReceivedUserCancel();
+		ImportContext.bUserCancelled |= DatasmithImporterImpl::HasUserCancelledTask( ImportContext.FeedbackContext );
 
 		TSharedPtr< IDatasmithLevelVariantSetsElement > LevelVariantSetsElement = ImportContext.FilteredScene->GetLevelVariantSets( LevelVariantSetIndex );
 		if ( !LevelVariantSetsElement )
@@ -2013,8 +2079,9 @@ void FDatasmithImporter::ImportLevelVariantSets( FDatasmithImportContext& Import
 			}
 		}
 
-		FString LevelVariantSetsName = ObjectTools::SanitizeObjectName( LevelVariantSetsElement->GetName() );
-		Progress.EnterProgressFrame( 1.f, FText::FromString( FString::Printf( TEXT("Importing level variant sets %d/%d (%s) ..."), LevelVariantSetIndex + 1, LevelVariantSetsCount, *LevelVariantSetsName ) ) );
+
+		FString LevelVariantSetsName = ObjectTools::SanitizeObjectName(LevelVariantSetsElement->GetName());
+		DatasmithImporterImpl::ReportProgress(Progress, 1.f, FText::FromString(FString::Printf(TEXT("Importing level variant sets %d/%d (%s) ..."), LevelVariantSetIndex + 1, LevelVariantSetsCount, *LevelVariantSetsName)));
 
 		ULevelVariantSets*& ImportedLevelVariantSets = ImportContext.ImportedLevelVariantSets.FindOrAdd( LevelVariantSetsElement.ToSharedRef() );
 		ImportedLevelVariantSets = FDatasmithLevelVariantSetsImporter::ImportLevelVariantSets( LevelVariantSetsElement.ToSharedRef(), ImportContext, ExistingLevelVariantSets );
@@ -2206,8 +2273,13 @@ void FDatasmithImporter::FinalizeImport(FDatasmithImportContext& ImportContext,
 	const int32 NumAssetsToFinalize = ValidAssets.Num() == 0 ? NumImportedObjects : ValidAssets.Num() + ImportContext.ImportedLevelSequences.Num() + ImportContext.ImportedLevelVariantSets.Num();
 	const int32 NumStaticMeshToBuild = ImportContext.ImportedStaticMeshes.Num();
 
-	FScopedSlowTask Progress((float)NumAssetsToFinalize + NumStaticMeshToBuild, LOCTEXT("FinalizingAssets", "Finalizing Assets"), true, *ImportContext.Warn);
-	Progress.MakeDialog(true);
+	TUniquePtr<FScopedSlowTask> ProgressPtr;
+
+	if ( ImportContext.FeedbackContext )
+	{ 
+		ProgressPtr = MakeUnique<FScopedSlowTask>((float)NumAssetsToFinalize + NumStaticMeshToBuild, LOCTEXT("FinalizingAssets", "Finalizing Assets"), true, *ImportContext.FeedbackContext);
+		ProgressPtr->MakeDialog(true);
+	}
 
 	TMap<UObject*, UObject*> ReferencesToRemap;
 
@@ -2219,6 +2291,8 @@ void FDatasmithImporter::FinalizeImport(FDatasmithImportContext& ImportContext,
 	const FString& RootFolderPath = ImportContext.AssetsContext.RootFolderPath;
 	const FString& TransientFolderPath = ImportContext.AssetsContext.TransientFolderPath;
 
+	FScopedSlowTask* Progress = ProgressPtr.Get();
+
 	// Needs to be done in dependencies order (textures -> materials -> static meshes)
 	for (const TPair< TSharedRef< IDatasmithTextureElement >, UTexture* >& ImportedTexturePair : ImportContext.ImportedTextures)
 	{
@@ -2236,7 +2310,7 @@ void FDatasmithImporter::FinalizeImport(FDatasmithImportContext& ImportContext,
 
 		FName TextureId = ImportedTexturePair.Key->GetName();
 
-		Progress.EnterProgressFrame(1.f, FText::FromString(FString::Printf(TEXT("Finalizing assets %d/%d (Texture %s) ..."), ++AssetIndex, NumAssetsToFinalize, *SourceTexture->GetName())));
+		DatasmithImporterImpl::ReportProgress(Progress, 1.f, FText::FromString(FString::Printf(TEXT("Finalizing assets %d/%d (Texture %s) ..."), ++AssetIndex, NumAssetsToFinalize, *SourceTexture->GetName())));
 
 		TSoftObjectPtr< UTexture >& ExistingTexturePtr = ImportContext.SceneAsset->Textures.FindOrAdd(TextureId);
 		UTexture* ExistingTexture = ExistingTexturePtr.Get();
@@ -2272,8 +2346,7 @@ void FDatasmithImporter::FinalizeImport(FDatasmithImportContext& ImportContext,
 		}
 
 		FName MaterialFunctionId = ImportedMaterialFunctionPair.Key->GetName();
-
-		Progress.EnterProgressFrame(1.f, FText::FromString(FString::Printf(TEXT("Finalizing assets %d/%d (Material Function %s) ..."), ++AssetIndex, NumAssetsToFinalize, *SourceMaterialFunction->GetName())));
+		DatasmithImporterImpl::ReportProgress(Progress, 1.f, FText::FromString(FString::Printf(TEXT("Finalizing assets %d/%d (Material Function %s) ..."), ++AssetIndex, NumAssetsToFinalize, *SourceMaterialFunction->GetName())));
 
 		TSoftObjectPtr< UMaterialFunction >& ExistingMaterialFunctionPtr = ImportContext.SceneAsset->MaterialFunctions.FindOrAdd(MaterialFunctionId);
 		UMaterialFunction* ExistingMaterialFunction = ExistingMaterialFunctionPtr.Get();
@@ -2301,7 +2374,7 @@ void FDatasmithImporter::FinalizeImport(FDatasmithImportContext& ImportContext,
 
 		FName MaterialId = ImportedMaterialPair.Key->GetName();
 
-		Progress.EnterProgressFrame(1.f, FText::FromString(FString::Printf(TEXT("Finalizing assets %d/%d (Material %s) ..."), ++AssetIndex, NumAssetsToFinalize, *SourceMaterialInterface->GetName())));
+		DatasmithImporterImpl::ReportProgress(Progress, 1.f, FText::FromString(FString::Printf(TEXT("Finalizing assets %d/%d (Material %s) ..."), ++AssetIndex, NumAssetsToFinalize, *SourceMaterialInterface->GetName())));
 
 		TSoftObjectPtr< UMaterialInterface >& ExistingMaterialPtr = ImportContext.SceneAsset->Materials.FindOrAdd(MaterialId);
 		UMaterialInterface* ExistingMaterial = ExistingMaterialPtr.Get();
@@ -2361,7 +2434,7 @@ void FDatasmithImporter::FinalizeImport(FDatasmithImportContext& ImportContext,
 
 		FName StaticMeshId = ImportedStaticMeshPair.Key->GetName();
 
-		Progress.EnterProgressFrame(1.f, FText::FromString(FString::Printf(TEXT("Finalizing assets %d/%d (Static Mesh %s) ..."), ++AssetIndex, NumAssetsToFinalize, *SourceStaticMesh->GetName())));
+		DatasmithImporterImpl::ReportProgress(Progress, 1.f, FText::FromString(FString::Printf(TEXT("Finalizing assets %d/%d (Static Mesh %s) ..."), ++AssetIndex, NumAssetsToFinalize, *SourceStaticMesh->GetName())));
 
 		TSoftObjectPtr< UStaticMesh >& ExistingStaticMeshPtr = ImportContext.SceneAsset->StaticMeshes.FindOrAdd(StaticMeshId);
 		UStaticMesh* ExistingStaticMesh = ExistingStaticMeshPtr.Get();
@@ -2379,7 +2452,7 @@ void FDatasmithImporter::FinalizeImport(FDatasmithImportContext& ImportContext,
 	int32 StaticMeshIndex = 0;
 	auto ProgressFunction = [&](UStaticMesh* StaticMesh)
 	{
-		Progress.EnterProgressFrame(1.f, FText::FromString(FString::Printf(TEXT("Building Static Mesh %d/%d (%s) ..."), ++StaticMeshIndex, StaticMeshes.Num(), *StaticMesh->GetName())));
+		DatasmithImporterImpl::ReportProgress(Progress, 1.f, FText::FromString(FString::Printf(TEXT("Building Static Mesh %d/%d (%s) ..."), ++StaticMeshIndex, StaticMeshes.Num(), *StaticMesh->GetName())));
 		return !ImportContext.bUserCancelled;
 	};
 
@@ -2401,7 +2474,8 @@ void FDatasmithImporter::FinalizeImport(FDatasmithImportContext& ImportContext,
 
 		FName LevelSequenceId = ImportedLevelSequencePair.Key->GetName();
 
-		Progress.EnterProgressFrame(1.f, FText::FromString(FString::Printf(TEXT("Finalizing assets %d/%d (Level Sequence %s) ..."), ++AssetIndex, NumAssetsToFinalize, *SourceLevelSequence->GetName())));
+		DatasmithImporterImpl::ReportProgress(Progress, 1.f, FText::FromString(FString::Printf(TEXT("Finalizing assets %d/%d (Level Sequence %s) ..."), ++AssetIndex, NumAssetsToFinalize, *SourceLevelSequence->GetName())));
+
 
 		TSoftObjectPtr< ULevelSequence >& ExistingLevelSequencePtr = ImportContext.SceneAsset->LevelSequences.FindOrAdd(LevelSequenceId);
 		ULevelSequence* ExistingLevelSequence = ExistingLevelSequencePtr.Get();
@@ -2429,7 +2503,7 @@ void FDatasmithImporter::FinalizeImport(FDatasmithImportContext& ImportContext,
 
 		FName LevelVariantSetsId = ImportedLevelVariantSetsPair.Key->GetName();
 
-		Progress.EnterProgressFrame(1.f, FText::FromString(FString::Printf(TEXT("Finalizing assets %d/%d (Level Variant Sets %s) ..."), ++AssetIndex, NumAssetsToFinalize, *SourceLevelVariantSets->GetName())));
+		DatasmithImporterImpl::ReportProgress(Progress, 1.f, FText::FromString(FString::Printf(TEXT("Finalizing assets %d/%d (Level Variant Sets %s) ..."), ++AssetIndex, NumAssetsToFinalize, *SourceLevelVariantSets->GetName())));
 
 		TSoftObjectPtr< ULevelVariantSets >& ExistingLevelVariantSetsPtr = ImportContext.SceneAsset->LevelVariantSets.FindOrAdd(LevelVariantSetsId);
 		ULevelVariantSets* ExistingLevelVariantSets = ExistingLevelVariantSetsPtr.Get();
@@ -2457,11 +2531,11 @@ void FDatasmithImporter::FinalizeImport(FDatasmithImportContext& ImportContext,
 
 	if ( ImportContext.ShouldImportActors() )
 	{
-		FinalizeActors(ImportContext, &ReferencesToRemap);
+		FinalizeActors( ImportContext, &ReferencesToRemap );
 	}
 
 	// Everything has been finalized, make sure the UDatasmithScene is set to dirty
-	if (ImportContext.SceneAsset)
+	if ( ImportContext.SceneAsset )
 	{
 		ImportContext.SceneAsset->MarkPackageDirty();
 	}
diff --git a/Engine/Plugins/Enterprise/DatasmithImporter/Source/DatasmithImporter/Private/DatasmithStaticMeshImporter.cpp b/Engine/Plugins/Enterprise/DatasmithImporter/Source/DatasmithImporter/Private/DatasmithStaticMeshImporter.cpp
--- a/Engine/Plugins/Enterprise/DatasmithImporter/Source/DatasmithImporter/Private/DatasmithStaticMeshImporter.cpp
+++ b/Engine/Plugins/Enterprise/DatasmithImporter/Source/DatasmithImporter/Private/DatasmithStaticMeshImporter.cpp
@@ -125,7 +125,7 @@ UStaticMesh* FDatasmithStaticMeshImporter::ImportStaticMesh(const TSharedRef< ID
 	{
 		if ( ExistingMesh->GetOuter() != Outer )
 		{
-			ResultStaticMesh = DuplicateObject< UStaticMesh >( ExistingMesh, Outer, *StaticMeshName );
+			ResultStaticMesh = FDatasmithImporterUtils::DuplicateObject< UStaticMesh >( ExistingMesh, Outer, *StaticMeshName );
 			IDatasmithImporterModule::Get().ResetOverrides( ResultStaticMesh ); // Don't copy the existing overrides
 		}
 		else
@@ -134,6 +134,17 @@ UStaticMesh* FDatasmithStaticMeshImporter::ImportStaticMesh(const TSharedRef< ID
 		}
 
 		ResultStaticMesh->SetFlags( ObjectFlags );
+
+		// Get rid of data that is going to be recreated anyway
+		for (FStaticMeshSourceModel& SourceModel : ResultStaticMesh->GetSourceModels())
+		{
+			SourceModel.MeshDescription.Reset();
+			SourceModel.MeshDescriptionBulkData.Reset();
+			if (SourceModel.RawMeshBulkData)
+			{
+				SourceModel.RawMeshBulkData->Empty();
+			}
+		}
 	}
 	else
 	{
@@ -306,8 +317,12 @@ void FDatasmithStaticMeshImporter::PreBuildStaticMeshes( FDatasmithImportContext
 {
 	TRACE_CPUPROFILER_EVENT_SCOPE(FDatasmithStaticMeshImporter::PreBuildStaticMeshes);
 
-	FScopedSlowTask Progress(ImportContext.ImportedStaticMeshes.Num(), LOCTEXT("PreBuildStaticMeshes", "Setting up UVs..."), true, *ImportContext.Warn);
-	Progress.MakeDialog(true);
+	TUniquePtr<FScopedSlowTask> ProgressPtr;
+	if ( ImportContext.FeedbackContext )
+	{ 
+		ProgressPtr = MakeUnique<FScopedSlowTask>(ImportContext.ImportedStaticMeshes.Num(), LOCTEXT("PreBuildStaticMeshes", "Setting up UVs..."), true, *ImportContext.FeedbackContext);
+		ProgressPtr->MakeDialog(true);
+	}
 
 	IMeshUtilities& MeshUtilities = FModuleManager::Get().LoadModuleChecked< IMeshUtilities >( "MeshUtilities" );
 
@@ -345,12 +360,17 @@ void FDatasmithStaticMeshImporter::PreBuildStaticMeshes( FDatasmithImportContext
 		);
 	}
 
+	FScopedSlowTask* Progress = ProgressPtr.Get();
+
 	// Ensure UI stays responsive by updating the progress even when the number of tasks hasn't changed
 	for (int32 OldTasksDone = 0, NewTasksDone = 0; OldTasksDone != SortedMesh.Num(); OldTasksDone = NewTasksDone)
 	{
 		NewTasksDone = TasksDone.Load();
-		Progress.EnterProgressFrame(NewTasksDone - OldTasksDone, FText::FromString(FString::Printf(TEXT("Packing UVs and computing tangents for static mesh %d/%d ..."), NewTasksDone, SortedMesh.Num())));
-		FPlatformProcess::Sleep(0.1);
+		if ( Progress )
+		{
+			Progress->EnterProgressFrame(NewTasksDone - OldTasksDone, FText::FromString(FString::Printf(TEXT("Packing UVs and computing tangents for static mesh %d/%d ..."), NewTasksDone, SortedMesh.Num())));
+			FPlatformProcess::Sleep(0.01);
+		}
 	}
 
 	for (int32 Index = 0; Index < TasksIsMeshValidResult.Num(); ++Index)
diff --git a/Engine/Plugins/Enterprise/DatasmithImporter/Source/DatasmithImporter/Private/DatasmithTextureImporter.cpp b/Engine/Plugins/Enterprise/DatasmithImporter/Source/DatasmithImporter/Private/DatasmithTextureImporter.cpp
--- a/Engine/Plugins/Enterprise/DatasmithImporter/Source/DatasmithImporter/Private/DatasmithTextureImporter.cpp
+++ b/Engine/Plugins/Enterprise/DatasmithImporter/Source/DatasmithImporter/Private/DatasmithTextureImporter.cpp
@@ -274,12 +274,30 @@ UTexture* FDatasmithTextureImporter::CreateTexture(const TSharedPtr<IDatasmithTe
 		// Notify the asset registry
 		FAssetRegistryModule::AssetCreated(Texture);
 
+		bool bUpdateResource = false;
+
 		if (FMath::IsNearlyEqual(RGBCurve, 1.0f) == false && RGBCurve > 0.f)
 		{
 			Texture->AdjustRGBCurve = RGBCurve;
-			Texture->UpdateResource();
+			bUpdateResource = true;
+		}
+
+		EDatasmithColorSpace ColorSpace = TextureElement->GetSRGB();
+		if (!Texture->SRGB && ColorSpace == EDatasmithColorSpace::sRGB)
+		{
+			Texture->SRGB = true;
+			bUpdateResource = true;
+		}
+		else if (Texture->SRGB && ColorSpace == EDatasmithColorSpace::Linear)
+		{
+			Texture->SRGB = false;
+			bUpdateResource = true;
 		}
 
+		if (bUpdateResource)
+		{
+			Texture->UpdateResource();
+		}
 		Texture->MarkPackageDirty();
 	}
 
diff --git a/Engine/Plugins/Enterprise/DatasmithImporter/Source/DatasmithImporter/Private/ObjectElements/DatasmithUObjectElements.cpp b/Engine/Plugins/Enterprise/DatasmithImporter/Source/DatasmithImporter/Private/ObjectElements/DatasmithUObjectElements.cpp
--- a/Engine/Plugins/Enterprise/DatasmithImporter/Source/DatasmithImporter/Private/ObjectElements/DatasmithUObjectElements.cpp
+++ b/Engine/Plugins/Enterprise/DatasmithImporter/Source/DatasmithImporter/Private/ObjectElements/DatasmithUObjectElements.cpp
@@ -916,6 +916,18 @@ void UDatasmithTextureElement::SetRGBCurve(float InValue)
 	return Element->SetRGBCurve(InValue);
 }
 
+EDatasmithColorSpace UDatasmithTextureElement::GetColorSpace() const
+{
+	DATASMITHOBJECTELEMENT_GETSHARED_AND_EARLYRETURN(IDatasmithTextureElement, TextureElemement, EDatasmithColorSpace::Default);
+	return Element->GetSRGB();
+}
+
+void UDatasmithTextureElement::SetColorSpace(EDatasmithColorSpace Option)
+{
+	DATASMITHOBJECTELEMENT_GETSHARED_AND_EARLYRETURN(IDatasmithTextureElement, TextureElemement, );
+	return Element->SetSRGB(Option);
+}
+
 bool UDatasmithTextureElement::IsElementValid() const
 {
 	DATASMITHOBJECTELEMENT_ISELEMENTVALID(TextureElemement);
diff --git a/Engine/Plugins/Enterprise/DatasmithImporter/Source/DatasmithImporter/Private/Utility/DatasmithImporterUtils.cpp b/Engine/Plugins/Enterprise/DatasmithImporter/Source/DatasmithImporter/Private/Utility/DatasmithImporterUtils.cpp
--- a/Engine/Plugins/Enterprise/DatasmithImporter/Source/DatasmithImporter/Private/Utility/DatasmithImporterUtils.cpp
+++ b/Engine/Plugins/Enterprise/DatasmithImporter/Source/DatasmithImporter/Private/Utility/DatasmithImporterUtils.cpp
@@ -1086,6 +1086,52 @@ const TSharedPtr<IDatasmithBaseMaterialElement>& FDatasmithImporterUtils::FDatas
 	}
 }
 
+UObject* FDatasmithImporterUtils::StaticDuplicateObject(UObject* SourceObject, UObject* Outer, const FName Name)
+{
+	UStaticMesh* SourceMesh = Cast< UStaticMesh >( SourceObject );
+
+	TArray<FStaticMeshSourceModel> SourceModels;
+	if (SourceMesh)
+	{
+		// Since static mesh can be quite heavy, remove source models for cloning to
+		// reduce useless work. Will be reinserted on the new duplicated asset.
+		SourceModels = MoveTemp(SourceMesh->GetSourceModels());
+
+		// Temporary flag to skip Postload during DuplicateObject
+		SourceMesh->SetFlags(RF_ArchetypeObject);
+	}
+
+	// Duplicate is used only to move our object from its temporary package into its final package replacing any asset
+	// already at that location. This function also takes care of fixing internal dependencies among the object's children.
+	// Since Duplicate has some rather heavy consequence, like calling PostLoad and doing all kind of stuff on an object
+	// that is not even fully initialized yet, we might want to find an alternative way of moving our objects in future 
+	// releases but keep it for the current release cycle.
+	UObject* Duplicate = ::DuplicateObject< UObject >( SourceObject, Outer, Name );
+
+	if ( SourceMesh )
+	{
+		UStaticMesh* DuplicateMesh = Cast< UStaticMesh >( Duplicate );
+
+		// Get rid of our temporary flag
+		SourceMesh->ClearFlags(RF_ArchetypeObject);
+		DuplicateMesh->ClearFlags(RF_ArchetypeObject);
+
+		for (FStaticMeshSourceModel& SourceModel : SourceModels)
+		{
+			// Fixup the new SourceModels owner
+			SourceModel.StaticMeshOwner = DuplicateMesh;
+
+			// Restore settings to the SourceMesh since they are used for diffing templates
+			SourceMesh->AddSourceModel().BuildSettings = SourceModel.BuildSettings;
+		}
+
+		// Apply source models to the duplicated mesh
+		DuplicateMesh->GetSourceModels() = MoveTemp(SourceModels);
+	}
+
+	return Duplicate;
+}
+
 FScopedLogger::FScopedLogger(FName LogTitle, const FText& LogLabel)
 	: Title(LogTitle)
 	, MessageLogModule(FModuleManager::LoadModuleChecked<FMessageLogModule>("MessageLog"))
diff --git a/Engine/Plugins/Enterprise/DatasmithImporter/Source/DatasmithImporter/Public/DatasmithImportContext.h b/Engine/Plugins/Enterprise/DatasmithImporter/Source/DatasmithImporter/Public/DatasmithImportContext.h
--- a/Engine/Plugins/Enterprise/DatasmithImporter/Source/DatasmithImporter/Public/DatasmithImportContext.h
+++ b/Engine/Plugins/Enterprise/DatasmithImporter/Source/DatasmithImporter/Public/DatasmithImportContext.h
@@ -276,7 +276,7 @@ struct DATASMITHIMPORTER_API FDatasmithImportContext
 	TMap< TSharedRef< IDatasmithLevelVariantSetsElement >, ULevelVariantSets* > ImportedLevelVariantSets;
 
 	/** Feedback context provided by ImportFactory. Used to display import progress */
-	FFeedbackContext* Warn;
+	FFeedbackContext* FeedbackContext;
 
 	/** Actors specific import context */
 	FDatasmithActorImportContext ActorsContext;
diff --git a/Engine/Plugins/Enterprise/DatasmithImporter/Source/DatasmithImporter/Public/ObjectElements/DatasmithUObjectElements.h b/Engine/Plugins/Enterprise/DatasmithImporter/Source/DatasmithImporter/Public/ObjectElements/DatasmithUObjectElements.h
--- a/Engine/Plugins/Enterprise/DatasmithImporter/Source/DatasmithImporter/Public/ObjectElements/DatasmithUObjectElements.h
+++ b/Engine/Plugins/Enterprise/DatasmithImporter/Source/DatasmithImporter/Public/ObjectElements/DatasmithUObjectElements.h
@@ -714,6 +714,14 @@ public:
 	UFUNCTION(BlueprintCallable, Category="Datasmith | Element")
 	void SetRGBCurve(float InRGBCurve);
 
+	/** Gets the color space of the texture */
+	UFUNCTION(BlueprintCallable, Category="Datasmith | Element")
+	EDatasmithColorSpace GetColorSpace() const;
+
+	/** Sets the color space of the texture */
+	UFUNCTION(BlueprintCallable, Category="Datasmith | Element")
+	void SetColorSpace(EDatasmithColorSpace Option);
+
 public:
 	TWeakPtr<IDatasmithTextureElement> GetDatasmithTextureElement() const { return TextureElemement; }
 	void SetDatasmithTextureElement(const TSharedPtr<IDatasmithTextureElement>& InElement) { TextureElemement = InElement; }
diff --git a/Engine/Plugins/Enterprise/DatasmithImporter/Source/DatasmithImporter/Public/Utility/DatasmithImporterUtils.h b/Engine/Plugins/Enterprise/DatasmithImporter/Source/DatasmithImporter/Public/Utility/DatasmithImporterUtils.h
--- a/Engine/Plugins/Enterprise/DatasmithImporter/Source/DatasmithImporter/Public/Utility/DatasmithImporterUtils.h
+++ b/Engine/Plugins/Enterprise/DatasmithImporter/Source/DatasmithImporter/Public/Utility/DatasmithImporterUtils.h
@@ -158,6 +158,23 @@ public:
 		int32 CurrentIndex;
 		TArray<TSharedPtr<IDatasmithBaseMaterialElement>> SortedMaterials;
 	};
+
+	/**
+	 * Convenience function duplicating an object specifically optimized for datasmith use cases
+	 *
+	 * @param SourceObject the object being copied
+	 * @param Outer the outer to use for the object
+	 * @param Name the optional name of the object
+	 *
+	 * @return the copied object or null if it failed for some reason
+	 */
+	static UObject* StaticDuplicateObject(UObject* SourceObject, UObject* Outer, const FName Name = NAME_None);
+
+	template< class T >
+	static T* DuplicateObject(T* SourceObject, UObject* Outer, const FName Name = NAME_None)
+	{
+		return (T*)FDatasmithImporterUtils::StaticDuplicateObject(SourceObject, Outer, Name);
+	}
 };
 
 template< typename ObjectType >
diff --git a/Engine/Plugins/Enterprise/VariantManager/Source/VariantManager/Private/DisplayNodes/VariantManagerActorNode.cpp b/Engine/Plugins/Enterprise/VariantManager/Source/VariantManager/Private/DisplayNodes/VariantManagerActorNode.cpp
--- a/Engine/Plugins/Enterprise/VariantManager/Source/VariantManager/Private/DisplayNodes/VariantManagerActorNode.cpp
+++ b/Engine/Plugins/Enterprise/VariantManager/Source/VariantManager/Private/DisplayNodes/VariantManagerActorNode.cpp
@@ -2,26 +2,31 @@
 
 #include "DisplayNodes/VariantManagerActorNode.h"
 
-#include "Widgets/DeclarativeSyntaxSupport.h"
-#include "Textures/SlateIcon.h"
+#include "SVariantManager.h"
+#include "Variant.h"
+#include "VariantManager.h"
+#include "VariantManagerDragDropOp.h"
+#include "VariantManagerEditorCommands.h"
+#include "VariantManagerSelection.h"
+#include "VariantObjectBinding.h"
+
+#include "Containers/ArrayBuilder.h"
+#include "DragAndDrop/ActorDragDropGraphEdOp.h"
+#include "EditorStyleSet.h"
 #include "Framework/Application/SlateApplication.h"
 #include "Framework/Commands/UIAction.h"
 #include "Framework/Commands/UICommandList.h"
 #include "Framework/MultiBox/MultiBoxBuilder.h"
-#include "EditorStyleSet.h"
 #include "GameFramework/Actor.h"
-#include "Containers/ArrayBuilder.h"
-#include "Styling/SlateIconFinder.h"
-#include "VariantManager.h"
-#include "VariantObjectBinding.h"
-#include "ScopedTransaction.h"
-#include "VariantManagerEditorCommands.h"
 #include "Input/DragAndDrop.h"
-#include "DragAndDrop/ActorDragDropGraphEdOp.h"
-#include "VariantManagerDragDropOp.h"
-#include "Variant.h"
+#include "Modules/ModuleManager.h"
+#include "SceneOutlinerModule.h"
+#include "SceneOutlinerPublicTypes.h"
+#include "ScopedTransaction.h"
+#include "Styling/SlateIconFinder.h"
+#include "Textures/SlateIcon.h"
 #include "UObject/StrongObjectPtr.h"
-#include "SVariantManager.h"
+#include "Widgets/DeclarativeSyntaxSupport.h"
 
 #define LOCTEXT_NAMESPACE "FVariantManagerActorNode"
 
@@ -32,6 +37,11 @@ FVariantManagerActorNode::FVariantManagerActorNode(UVariantObjectBinding* InObje
 {
 }
 
+TWeakObjectPtr<UVariantObjectBinding> FVariantManagerActorNode::GetObjectBinding() const
+{
+	return ObjectBinding;
+}
+
 void FVariantManagerActorNode::BuildContextMenu(FMenuBuilder& MenuBuilder)
 {
 	FVariantManagerDisplayNode::BuildContextMenu(MenuBuilder);
@@ -43,23 +53,47 @@ void FVariantManagerActorNode::BuildContextMenu(FMenuBuilder& MenuBuilder)
 	MenuBuilder.BeginSection("Actor", LOCTEXT("ActorSectionText", "Actor"));
 	MenuBuilder.AddMenuEntry(FVariantManagerEditorCommands::Get().AddPropertyCaptures);
 	MenuBuilder.AddMenuEntry(FVariantManagerEditorCommands::Get().AddFunction);
+
+	int32 NumSelectedNodes = 0;
+	if (TSharedPtr<FVariantManager> PinnedVariantManager = VariantManager.Pin())
+	{
+		FVariantManagerSelection& Selection = PinnedVariantManager->GetSelection();
+		NumSelectedNodes = Selection.GetSelectedActorNodes().Num();
+	}
+
+	if (NumSelectedNodes > 1)
+	{
+		MenuBuilder.AddMenuEntry(FVariantManagerEditorCommands::Get().RebindActorDisabled);
+	}
+	else
+	{
+		MenuBuilder.AddSubMenu(
+			LOCTEXT("RebindActorName", "Rebind to other Actor"),
+			LOCTEXT("RebindActorName", "Connect this binding to a different actor and try reusing the captured properties"),
+			FNewMenuDelegate::CreateSP(this, &FVariantManagerActorNode::AddAssignActorSubMenu));
+	}
+
 	MenuBuilder.EndSection();
 }
 
 FText FVariantManagerActorNode::GetDisplayNameToolTipText() const
 {
-	//FSequencer& Sequencer = ParentTree.GetSequencer();
-	//if (Sequencer.FindObjectsInCurrentSequence(ObjectBinding).Num() == 0)
-	//{
-	//	return LOCTEXT("InvalidBoundObjectToolTip", "The object bound to this track is missing.");
-	//}
+	if (const UClass* BindingClass = GetClassForObjectBinding())
+	{
+		return FText::FromString(BindingClass->GetName());
+	}
 
-	return FText();
+	return LOCTEXT("FailedToResolveTooltip", "Binding can't be resolved. Right-click to rebind to another actor");
 }
 
 const FSlateBrush* FVariantManagerActorNode::GetIconBrush() const
 {
-	return FSlateIconFinder::FindIconBrushForClass(GetClassForObjectBinding());
+	if (const UClass* BindingClass = GetClassForObjectBinding())
+	{
+		return FSlateIconFinder::FindIconBrushForClass(BindingClass);
+	}
+
+	return FEditorStyle::Get().GetBrush(TEXT("MessageLog.Warning"));
 }
 
 const FSlateBrush* FVariantManagerActorNode::GetIconOverlayBrush() const
@@ -67,11 +101,6 @@ const FSlateBrush* FVariantManagerActorNode::GetIconOverlayBrush() const
 	return nullptr;
 }
 
-FText FVariantManagerActorNode::GetIconToolTipText() const
-{
-	return LOCTEXT("PossessableToolTip", "This item is a possessable reference to an existing object.");
-}
-
 EVariantManagerNodeType FVariantManagerActorNode::GetType() const
 {
 	return EVariantManagerNodeType::Actor;
@@ -86,8 +115,7 @@ FText FVariantManagerActorNode::GetDisplayName() const
 	// we stop showing properties if our actor doesn't resolve anymore, though
 	if (NewDisplayText.ToString() != OldDisplayText.ToString())
 	{
-		TSharedPtr<FVariantManager> VarMan = GetVariantManager().Pin();
-		if (VarMan.IsValid())
+		if (TSharedPtr<FVariantManager> VarMan = GetVariantManager().Pin())
 		{
 			VarMan->GetVariantManagerWidget()->RefreshPropertyList();
 		}
@@ -97,12 +125,21 @@ FText FVariantManagerActorNode::GetDisplayName() const
 	return NewDisplayText;
 }
 
+FSlateColor FVariantManagerActorNode::GetDisplayNameColor() const
+{
+	if (GetClassForObjectBinding())
+	{
+		return FSlateColor(FLinearColor(1.0f, 1.0f, 1.0f, 1.0f));
+	}
+
+	return FSlateColor(FLinearColor(0.6f, 0.6f, 0.6f, 1.0f));
+}
+
 void FVariantManagerActorNode::SetDisplayName(const FText& NewDisplayName)
 {
 	UObject* Obj = ObjectBinding->GetObject();
 
-	AActor* ObjAsActor = Cast<AActor>(Obj);
-	if (ObjAsActor != nullptr)
+	if (AActor* ObjAsActor = Cast<AActor>(Obj))
 	{
 		ObjAsActor->SetActorLabel(*NewDisplayName.ToString());
 	}
@@ -122,6 +159,11 @@ bool FVariantManagerActorNode::CanDrag() const
 	return true;
 }
 
+TWeakPtr<FVariantManager> FVariantManagerActorNode::GetVariantManager() const
+{
+	return VariantManager;
+}
+
 TOptional<EItemDropZone> FVariantManagerActorNode::CanDrop(const FDragDropEvent& DragDropEvent, EItemDropZone ItemDropZone) const
 {
 	TSharedPtr<FActorDragDropGraphEdOp> ActorDragDrop = DragDropEvent.GetOperationAs<FActorDragDropGraphEdOp>();
@@ -178,11 +220,9 @@ TOptional<EItemDropZone> FVariantManagerActorNode::CanDrop(const FDragDropEvent&
 		{
 			if (DraggedNode->GetType() == EVariantManagerNodeType::Actor)
 			{
-				TSharedPtr<FVariantManagerActorNode> DraggedActorNode = StaticCastSharedRef<FVariantManagerActorNode>(DraggedNode);
-				if (DraggedActorNode.IsValid())
+				if (TSharedPtr<FVariantManagerActorNode> DraggedActorNode = StaticCastSharedRef<FVariantManagerActorNode>(DraggedNode))
 				{
-					UVariantObjectBinding* DraggedBinding = DraggedActorNode->GetObjectBinding().Get();
-					if (DraggedBinding)
+					if (UVariantObjectBinding* DraggedBinding = DraggedActorNode->GetObjectBinding().Get())
 					{
 						UObject* DraggedActor = DraggedBinding->GetObject();
 						if (!DraggedActors.Contains(DraggedActor))
@@ -214,8 +254,7 @@ TOptional<EItemDropZone> FVariantManagerActorNode::CanDrop(const FDragDropEvent&
 		TArray<TWeakObjectPtr<AActor>> ActorsToCheck;
 		for (UObject* Actor : DraggedActors)
 		{
-			AActor* ActualActor = Cast<AActor>(Actor);
-			if (ActualActor)
+			if (AActor* ActualActor = Cast<AActor>(Actor))
 			{
 				ActorsToCheck.Add(ActualActor);
 			}
@@ -422,34 +461,120 @@ void FVariantManagerActorNode::Drop(const FDragDropEvent& DragDropEvent, EItemDr
 
 const UClass* FVariantManagerActorNode::GetClassForObjectBinding() const
 {
-	UObject* Obj = ObjectBinding->GetObject();
-	if (Obj != nullptr)
+	if (ObjectBinding.IsValid())
 	{
-		return Obj->GetClass();
+		if (UObject* Obj = ObjectBinding->GetObject())
+		{
+			return Obj->GetClass();
+		}
 	}
 
 	return nullptr;
 }
 
-TSharedRef<SWidget> FVariantManagerActorNode::GetCustomOutlinerContent(TSharedPtr<SVariantManagerTableRow> InTableRow)
+void FVariantManagerActorNode::AddAssignActorSubMenu(FMenuBuilder& MenuBuilder)
 {
-	FSlateFontInfo NodeFont = FEditorStyle::GetFontStyle("Sequencer.AnimationOutliner.RegularFont");
-
-	EditableLabel = SNew(SInlineEditableTextBlock)
-		.IsReadOnly(this, &FVariantManagerDisplayNode::IsReadOnly)
-		.Font(NodeFont)
-		.ColorAndOpacity(this, &FVariantManagerDisplayNode::GetDisplayNameColor)
-		.OnTextCommitted(this, &FVariantManagerDisplayNode::HandleNodeLabelTextChanged)
-		.Text(this, &FVariantManagerDisplayNode::GetDisplayName)
-		.ToolTipText(this, &FVariantManagerDisplayNode::GetDisplayNameToolTipText)
-		.Clipping(EWidgetClipping::ClipToBounds);
-
-	FText Tooltip;
-	if (const UClass* BindingClass = GetClassForObjectBinding())
+	// Copied from FSequencer::AssignActor
+	using namespace SceneOutliner;
+
+	// If we're showing this menu, we know for a fact only our actor node is selected,
+	// so we only have to check our variant
+	TSet<const AActor*> BoundActors;
+	if (ObjectBinding.IsValid())
+	{
+		if (UVariant* ParentVariant = ObjectBinding->GetParent())
+		{
+			for (UVariantObjectBinding* Binding : ParentVariant->GetBindings())
+			{
+				if (AActor* BoundActor = Cast<AActor>(Binding->GetObject()))
+				{
+					BoundActors.Add(BoundActor);
+				}
+			}
+		}
+	}
+
+	auto IsActorValidForAssignment = [BoundActors](const AActor* InActor)
 	{
-		Tooltip = FText::FromString(BindingClass->GetName());
+		return !BoundActors.Contains(InActor);
+	};
+
+	MenuBuilder.AddMenuEntry(FVariantManagerEditorCommands::Get().RebindToSelected,
+							 NAME_None,
+							 TAttribute<FText>(),
+							 TAttribute<FText>(this, &FVariantManagerActorNode::GetRebindToSelectedTooltip));
+
+	// Set up a menu entry to assign an actor to the object binding node
+	FInitializationOptions InitOptions;
+	{
+		InitOptions.Mode = ESceneOutlinerMode::ActorPicker;
+
+		// We hide the header row to keep the UI compact.
+		InitOptions.bShowHeaderRow = false;
+		InitOptions.bShowSearchBox = true;
+		InitOptions.bShowCreateNewFolder = false;
+		InitOptions.bFocusSearchBoxWhenOpened = true;
+		// Only want the actor label column
+		InitOptions.ColumnMap.Add(FBuiltInColumnTypes::Label(), FColumnInfo(EColumnVisibility::Visible, 0));
+
+		// Only display actors that are not possessed already
+		InitOptions.Filters->AddFilterPredicate( FActorFilterPredicate::CreateLambda( IsActorValidForAssignment ) );
+	}
+
+	// actor selector to allow the user to choose an actor
+	FSceneOutlinerModule& SceneOutlinerModule = FModuleManager::LoadModuleChecked<FSceneOutlinerModule>("SceneOutliner");
+	TSharedRef< SWidget > MiniSceneOutliner =
+		SNew( SBox )
+		.MaxDesiredHeight(400.0f)
+		.WidthOverride(300.0f)
+		[
+			SceneOutlinerModule.CreateSceneOutliner(
+				InitOptions,
+				FOnActorPicked::CreateLambda([&](AActor* Actor)
+				{
+					FSlateApplication::Get().DismissAllMenus();
+
+					if (Actor && ObjectBinding.IsValid())
+					{
+						FScopedTransaction Transaction(FText::Format(
+							LOCTEXT("RebindToActorTransaction", "Rebind variant to actor '{0}'"),
+							FText::FromString(Actor->GetActorLabel())));
+
+						ObjectBinding->SetObject(Actor);
+					}
+				})
+			)
+		];
+
+	MenuBuilder.AddMenuSeparator();
+	MenuBuilder.AddWidget(MiniSceneOutliner, FText::GetEmpty(), true);
+}
+
+FText FVariantManagerActorNode::GetRebindToSelectedTooltip() const
+{
+	if (TSharedPtr<FVariantManager> PinnedVariantManager = VariantManager.Pin())
+	{
+		if(TSharedPtr<SVariantManager> VariantManagerWidget = PinnedVariantManager->GetVariantManagerWidget())
+		{
+			UVariantObjectBinding* SelectedBinding = nullptr;
+			UObject* SelectedObject = nullptr;
+			VariantManagerWidget->GetSelectedBindingAndEditorActor(SelectedBinding, SelectedObject);
+
+			AActor* SelectedActor = Cast<AActor>(SelectedObject);
+
+			if (SelectedBinding && SelectedActor)
+			{
+				return FText::Format(LOCTEXT("RebindToThisActorTooltipFound", "Rebind to '{0}'"), FText::FromString(SelectedActor->GetActorLabel()));
+			}
+		}
 	}
 
+	// If we reach here, then so did SVariantManager::CanRebindToSelectedActor, so the button will be disabled
+	return LOCTEXT("RebindToThisActorTooltipDisabled", "Cannot rebind to this selected actor");
+}
+
+TSharedRef<SWidget> FVariantManagerActorNode::GetCustomOutlinerContent(TSharedPtr<SVariantManagerTableRow> InTableRow)
+{
 	return
 	SNew(SBox)
 	[
@@ -464,6 +589,7 @@ TSharedRef<SWidget> FVariantManagerActorNode::GetCustomOutlinerContent(TSharedPt
 			.HAlign(HAlign_Left)
 			.Padding(FMargin( 4.0, 0.0f, 0.0f, 0.0f ))
 			.HeightOverride(26)
+			.ToolTipText(this, &FVariantManagerActorNode::GetDisplayNameToolTipText)
 			[
 				SNew(SHorizontalBox)
 
@@ -478,7 +604,6 @@ TSharedRef<SWidget> FVariantManagerActorNode::GetCustomOutlinerContent(TSharedPt
 					[
 						SNew(SImage)
 						.Image(this, &FVariantManagerActorNode::GetIconBrush)
-						.ToolTipText(Tooltip)
 					]
 				]
 
@@ -486,7 +611,14 @@ TSharedRef<SWidget> FVariantManagerActorNode::GetCustomOutlinerContent(TSharedPt
 				.FillWidth(1.0f)
 				.VAlign(VAlign_Center)
 				[
-					EditableLabel.ToSharedRef()
+					SAssignNew(EditableLabel, SInlineEditableTextBlock)
+					.IsReadOnly(this, &FVariantManagerDisplayNode::IsReadOnly)
+					.Font(FEditorStyle::GetFontStyle("Sequencer.AnimationOutliner.RegularFont"))
+					.ColorAndOpacity(this, &FVariantManagerDisplayNode::GetDisplayNameColor)
+					.OnTextCommitted(this, &FVariantManagerDisplayNode::HandleNodeLabelTextChanged)
+					.Text(this, &FVariantManagerDisplayNode::GetDisplayName)
+					.ToolTipText(this, &FVariantManagerDisplayNode::GetDisplayNameToolTipText)
+					.Clipping(EWidgetClipping::ClipToBounds)
 				]
 			]
 		]
diff --git a/Engine/Plugins/Enterprise/VariantManager/Source/VariantManager/Private/DisplayNodes/VariantManagerActorNode.h b/Engine/Plugins/Enterprise/VariantManager/Source/VariantManager/Private/DisplayNodes/VariantManagerActorNode.h
--- a/Engine/Plugins/Enterprise/VariantManager/Source/VariantManager/Private/DisplayNodes/VariantManagerActorNode.h
+++ b/Engine/Plugins/Enterprise/VariantManager/Source/VariantManager/Private/DisplayNodes/VariantManagerActorNode.h
@@ -2,69 +2,59 @@
 
 #pragma once
 
-#include "CoreMinimal.h"
-#include "Widgets/SWidget.h"
 #include "DisplayNodes/VariantManagerDisplayNode.h"
+
+#include "CoreMinimal.h"
 #include "PropertyPath.h"
+#include "Widgets/SWidget.h"
 
+class FDragDropEvent;
 class FMenuBuilder;
-struct FSlateBrush;
+class SVariantManagerTableRow;
 class UVariantObjectBinding;
 enum class EItemDropZone;
-class FDragDropEvent;
-class SVariantManagerTableRow;
+struct FSlateBrush;
 
 /**
 * A node for displaying an object binding
 */
-class FVariantManagerActorNode
-	: public FVariantManagerDisplayNode
+class FVariantManagerActorNode : public FVariantManagerDisplayNode
 {
 public:
-
 	FVariantManagerActorNode(UVariantObjectBinding* InObjectBinding, TSharedPtr<FVariantManagerDisplayNode> InParentNode, TWeakPtr<FVariantManager> InVariantManager);
 
-	/** @return The object binding on this node */
-	TWeakObjectPtr<UVariantObjectBinding> GetObjectBinding() const
-	{
-		return ObjectBinding;
-	}
+	/** Gets the UVariantObjectBinding held by this actor node */
+	TWeakObjectPtr<UVariantObjectBinding> GetObjectBinding() const;
 
-	// FVariantManagerDisplayNode interface
+	// Begin FVariantManagerDisplayNode interface
 	virtual void BuildContextMenu(FMenuBuilder& MenuBuilder) override;
 	virtual FText GetDisplayNameToolTipText() const override;
 	virtual const FSlateBrush* GetIconBrush() const override;
 	virtual const FSlateBrush* GetIconOverlayBrush() const override;
-	virtual FText GetIconToolTipText() const override;
 	virtual EVariantManagerNodeType GetType() const override;
 	virtual FText GetDisplayName() const override;
+	virtual FSlateColor GetDisplayNameColor() const override;
 	virtual void SetDisplayName(const FText& NewDisplayName) override;
 	virtual bool IsSelectable() const override;
 	virtual bool CanDrag() const override;
-	virtual TWeakPtr<FVariantManager> GetVariantManager() const override
-	{
-		return VariantManager;
-	}
-
+	virtual TWeakPtr<FVariantManager> GetVariantManager() const override;
 	virtual TOptional<EItemDropZone> CanDrop(const FDragDropEvent& DragDropEvent, EItemDropZone ItemDropZone) const override;
 	virtual void Drop(const FDragDropEvent& DragDropEvent, EItemDropZone ItemDropZone) override;
 	virtual TSharedRef<SWidget> GetCustomOutlinerContent(TSharedPtr<SVariantManagerTableRow> InTableRow) override;
+	// End FVariantManagerDisplayNode interface
 
-protected:
+private:
+	/** Gets the tooltip to show on the 'RebindToSelected' button */
+	FText GetRebindToSelectedTooltip() const;
 
+	/** Gets the class of the UObject that is bound to our ObjectBinding, or nullptr */
 	const UClass* GetClassForObjectBinding() const;
 
-private:
-
-	void HandleAddTrackSubMenuNew(FMenuBuilder& AddTrackMenuBuilder, TArray<FPropertyPath> KeyablePropertyPath, int32 PropertyNameIndexStart = 0);
-	void HandleLabelsSubMenuCreate(FMenuBuilder& MenuBuilder);
-	void HandlePropertyMenuItemExecute(FPropertyPath PropertyPath);
-	TSharedRef<SWidget> HandleAddTrackComboButtonGetMenuContent();
+	/** Creates the floating submenu that is shown when mousing over the "Rebind to other actor" right-click option */
+	void AddAssignActorSubMenu(FMenuBuilder& MenuBuilder);
 
 	TWeakObjectPtr<UVariantObjectBinding> ObjectBinding;
 	mutable FText OldDisplayText;
-
 	FText DefaultDisplayName;
-
 	TWeakPtr<FVariantManager> VariantManager;
 };
diff --git a/Engine/Plugins/Enterprise/VariantManager/Source/VariantManager/Private/DisplayNodes/VariantManagerEnumPropertyNode.cpp b/Engine/Plugins/Enterprise/VariantManager/Source/VariantManager/Private/DisplayNodes/VariantManagerEnumPropertyNode.cpp
--- a/Engine/Plugins/Enterprise/VariantManager/Source/VariantManager/Private/DisplayNodes/VariantManagerEnumPropertyNode.cpp
+++ b/Engine/Plugins/Enterprise/VariantManager/Source/VariantManager/Private/DisplayNodes/VariantManagerEnumPropertyNode.cpp
@@ -2,19 +2,20 @@
 
 #include "DisplayNodes/VariantManagerEnumPropertyNode.h"
 
-#include "GameFramework/Actor.h"
-#include "PropertyValue.h"
 #include "PropertyTemplateObject.h"
-#include "PropertyEditorModule.h"
-#include "Modules/ModuleManager.h"
-#include "Widgets/Input/SNumericEntryBox.h"
-#include "Widgets/Input/SComboBox.h"
-#include "ScopedTransaction.h"
-#include "IDocumentation.h"
+#include "PropertyValue.h"
 #include "SVariantManager.h"
 #include "VariantManagerLog.h"
 #include "VariantObjectBinding.h"
 
+#include "GameFramework/Actor.h"
+#include "IDocumentation.h"
+#include "Modules/ModuleManager.h"
+#include "PropertyEditorModule.h"
+#include "ScopedTransaction.h"
+#include "Widgets/Input/SComboBox.h"
+#include "Widgets/Input/SNumericEntryBox.h"
+
 #define LOCTEXT_NAMESPACE "FVariantManagerEnumPropertyNode"
 
 FVariantManagerEnumPropertyNode::FVariantManagerEnumPropertyNode(TArray<UPropertyValue*> InPropertyValues, TWeakPtr<FVariantManager> InVariantManager)
@@ -108,17 +109,8 @@ TSharedPtr<SWidget> FVariantManagerEnumPropertyNode::GetPropertyValueWidget()
 
 	UpdateComboboxStrings();
 
-	int32 EnumIndex = FirstPropertyValue->GetRecordedDataAsEnumIndex();
-	int32 ComboboxItemIndex = EnumIndices.Find(EnumIndex);
-	if (ComboboxItemIndex == INDEX_NONE)
-	{
-		UE_LOG(LogVariantManager, Warning, TEXT("For captured property '%s', did not find an UEnum item with index %d"), *FirstPropertyValue->GetPropertyName().ToString(), EnumIndex);
-		ComboboxItemIndex = 0;
-	}
-
 	SAssignNew(Combobox, SComboBox<TSharedPtr<FString>>)
 	.OptionsSource(&EnumDisplayTexts)
-	.InitiallySelectedItem(EnumDisplayTexts[ComboboxItemIndex])
 	.OnGenerateWidget_Lambda([bSameValue](TSharedPtr<FString> Item)
 	{
 		return SNew(STextBlock).Text(FText::FromString(*Item));
@@ -126,7 +118,7 @@ TSharedPtr<SWidget> FVariantManagerEnumPropertyNode::GetPropertyValueWidget()
 	.Content()
 	[
 		SNew(STextBlock)
-		.Text(this, &FVariantManagerEnumPropertyNode::ComboboxGetText, bSameValue)
+		.Text(this, &FVariantManagerEnumPropertyNode::GetRecordedEnumDisplayText, bSameValue)
 	]
 	.OnSelectionChanged(this, &FVariantManagerEnumPropertyNode::OnComboboxSelectionChanged); // Widget updates recorded data
 
@@ -179,22 +171,31 @@ void FVariantManagerEnumPropertyNode::OnComboboxSelectionChanged(TSharedPtr<FStr
 	GetVariantManager().Pin()->GetVariantManagerWidget()->RefreshPropertyList();
 }
 
-FText FVariantManagerEnumPropertyNode::ComboboxGetText(bool bSameValue) const
+FText FVariantManagerEnumPropertyNode::GetRecordedEnumDisplayText(bool bSameValue) const
 {
-	if (Combobox.IsValid())
+	if (!bSameValue)
 	{
-		if (!bSameValue)
-		{
-			return LOCTEXT("MultipleValuesLabel", "Multiple Values");
-		}
+		return LOCTEXT("MultipleValuesLabel", "Multiple Values");
+	}
+
+	// We don't just use Combobox->GetSelectedItem() here because the Combobox watches
+	// EnumDisplayTexts, and that list only contains the valid and non-hidden options for the Enum.
+	// The user can only choose between those, but the PropertyValue may have a hidden or 'invalid'
+	// option recorded already (e.g. _MAX), and with this we will display it as it actually is,
+	// which is also how the Details panel behaves
+	if (PropertyValues.Num() >= 1)
+	{
+		UPropertyValue* FirstPropertyValue = PropertyValues[0].Get();
 
-		TSharedPtr<FString> SelectedText = Combobox->GetSelectedItem();
-		if (SelectedText.IsValid())
+		if (UEnum* Enum = FirstPropertyValue->GetEnumPropertyEnum())
 		{
-			return FText::FromString(*SelectedText);
+			int32 EnumIndex = FirstPropertyValue->GetRecordedDataAsEnumIndex();
+			FText EnumDisplayNameText = Enum->GetDisplayNameTextByIndex(EnumIndex);
+			if(!EnumDisplayNameText.IsEmpty())
+			{
+				return EnumDisplayNameText;
+			}
 		}
-
-		return LOCTEXT("InvalidLabel", "(INVALID)");
 	}
 
 	return FText();
diff --git a/Engine/Plugins/Enterprise/VariantManager/Source/VariantManager/Private/DisplayNodes/VariantManagerEnumPropertyNode.h b/Engine/Plugins/Enterprise/VariantManager/Source/VariantManager/Private/DisplayNodes/VariantManagerEnumPropertyNode.h
--- a/Engine/Plugins/Enterprise/VariantManager/Source/VariantManager/Private/DisplayNodes/VariantManagerEnumPropertyNode.h
+++ b/Engine/Plugins/Enterprise/VariantManager/Source/VariantManager/Private/DisplayNodes/VariantManagerEnumPropertyNode.h
@@ -2,36 +2,29 @@
 
 #pragma once
 
-#include "CoreMinimal.h"
-#include "Widgets/SWidget.h"
 #include "DisplayNodes/VariantManagerPropertyNode.h"
+
+#include "CoreMinimal.h"
 #include "Widgets/Input/SComboBox.h"
+#include "Widgets/SWidget.h"
 
 class SToolTip;
-class SVariantManagerTableRow;
 
-
-class FVariantManagerEnumPropertyNode
-	: public FVariantManagerPropertyNode
+class FVariantManagerEnumPropertyNode : public FVariantManagerPropertyNode
 {
 public:
-
 	FVariantManagerEnumPropertyNode(TArray<UPropertyValue*> InPropertyValues, TWeakPtr<FVariantManager> InVariantManager);
 
 protected:
-
 	virtual TSharedPtr<SWidget> GetPropertyValueWidget() override;
 
 private:
-
 	void OnComboboxSelectionChanged(TSharedPtr<FString> NewItem, ESelectInfo::Type SelectType);
-	FText ComboboxGetText(bool bSameValue) const;
-
+	FText GetRecordedEnumDisplayText(bool bSameValue) const;
 	void UpdateComboboxStrings();
 
 	TSharedPtr<SComboBox<TSharedPtr<FString>>> Combobox;
 
-	// Also store indices because we won't store texts/tooltips for hidden enums
 	TArray<TSharedPtr<FString>> EnumDisplayTexts;
 	TArray<TSharedPtr<SToolTip>> EnumRichToolTips;
 	TArray<int32> EnumIndices;
diff --git a/Engine/Plugins/Enterprise/VariantManager/Source/VariantManager/Private/DisplayNodes/VariantManagerFunctionPropertyNode.cpp b/Engine/Plugins/Enterprise/VariantManager/Source/VariantManager/Private/DisplayNodes/VariantManagerFunctionPropertyNode.cpp
--- a/Engine/Plugins/Enterprise/VariantManager/Source/VariantManager/Private/DisplayNodes/VariantManagerFunctionPropertyNode.cpp
+++ b/Engine/Plugins/Enterprise/VariantManager/Source/VariantManager/Private/DisplayNodes/VariantManagerFunctionPropertyNode.cpp
@@ -62,6 +62,16 @@ FFunctionCaller& FVariantManagerFunctionPropertyNode::GetFunctionCaller() const
 	return FunctionCaller;
 }
 
+uint32 FVariantManagerFunctionPropertyNode::GetDisplayOrder() const
+{
+	return FunctionCaller.GetDisplayOrder();
+}
+
+void FVariantManagerFunctionPropertyNode::SetDisplayOrder(uint32 InDisplayOrder)
+{
+	FunctionCaller.SetDisplayOrder(InDisplayOrder);
+}
+
 void FVariantManagerFunctionPropertyNode::SetBindingTargetFunction(UK2Node_FunctionEntry* NewFunctionEntry)
 {
 	FScopedTransaction Transaction(LOCTEXT("SetFunctionCallerFunction", "Set a new function to be called by a FunctionCaller"));
diff --git a/Engine/Plugins/Enterprise/VariantManager/Source/VariantManager/Private/DisplayNodes/VariantManagerFunctionPropertyNode.h b/Engine/Plugins/Enterprise/VariantManager/Source/VariantManager/Private/DisplayNodes/VariantManagerFunctionPropertyNode.h
--- a/Engine/Plugins/Enterprise/VariantManager/Source/VariantManager/Private/DisplayNodes/VariantManagerFunctionPropertyNode.h
+++ b/Engine/Plugins/Enterprise/VariantManager/Source/VariantManager/Private/DisplayNodes/VariantManagerFunctionPropertyNode.h
@@ -29,6 +29,16 @@ public:
 	TWeakObjectPtr<UVariantObjectBinding> GetObjectBinding() const;
 	FFunctionCaller& GetFunctionCaller() const;
 
+	/**
+	* Get the order with which the VariantManager should display this in a property list. Lower values will be shown higher up
+	*/
+	virtual uint32 GetDisplayOrder() const override;
+
+	/**
+	* Set the order with which the VariantManager should display this in a property list. Lower values will be shown higher up
+	*/
+	virtual void SetDisplayOrder(uint32 InDisplayOrder) override;
+
 private:
 
 	// Set the target function as the selected option both for this node as well as for the underlying binding
diff --git a/Engine/Plugins/Enterprise/VariantManager/Source/VariantManager/Private/DisplayNodes/VariantManagerPropertyNode.cpp b/Engine/Plugins/Enterprise/VariantManager/Source/VariantManager/Private/DisplayNodes/VariantManagerPropertyNode.cpp
--- a/Engine/Plugins/Enterprise/VariantManager/Source/VariantManager/Private/DisplayNodes/VariantManagerPropertyNode.cpp
+++ b/Engine/Plugins/Enterprise/VariantManager/Source/VariantManager/Private/DisplayNodes/VariantManagerPropertyNode.cpp
@@ -2,30 +2,34 @@
 
 #include "DisplayNodes/VariantManagerPropertyNode.h"
 
-#include "Widgets/DeclarativeSyntaxSupport.h"
-#include "Framework/MultiBox/MultiBoxBuilder.h"
-#include "GameFramework/Actor.h"
-#include "EditorStyleSet.h"
-#include "PropertyEditorModule.h"
-#include "Modules/ModuleManager.h"
+#include "PropertyTemplateObject.h"
+#include "PropertyValue.h"
+#include "SVariantManager.h"
 #include "VariantManager.h"
+#include "VariantManagerDragDropOp.h"
+#include "VariantManagerEditorCommands.h"
 #include "VariantManagerLog.h"
 #include "VariantObjectBinding.h"
+
+#include "EditorFontGlyphs.h"
+#include "EditorStyleSet.h"
+#include "Framework/MultiBox/MultiBoxBuilder.h"
+#include "GameFramework/Actor.h"
 #include "ISinglePropertyView.h"
-#include "PropertyValue.h"
-#include "PropertyHandle.h"
-#include "PropertyTemplateObject.h"
-#include "VariantManagerEditorCommands.h"
+#include "Input/Reply.h"
+#include "Modules/ModuleManager.h"
 #include "PropertyCustomizationHelpers.h"
-#include "SVariantManager.h"
+#include "PropertyEditorModule.h"
+#include "PropertyHandle.h"
 #include "ScopedTransaction.h"
-#include "Widgets/Input/SButton.h"
-#include "Input/Reply.h"
+#include "Widgets/DeclarativeSyntaxSupport.h"
 #include "Widgets/Images/SImage.h"
-#include "EditorFontGlyphs.h"
+#include "Widgets/Input/SButton.h"
 
 #define LOCTEXT_NAMESPACE "FVariantManagerPropertyNode"
 
+using FDisplayNodeRef = TSharedRef<FVariantManagerDisplayNode>;
+
 FVariantManagerPropertyNode::FVariantManagerPropertyNode(TArray<UPropertyValue*> InPropertyValues, TWeakPtr<FVariantManager> InVariantManager)
 	: FVariantManagerDisplayNode(nullptr, nullptr)
 	, PropertyValues(InPropertyValues)
@@ -99,6 +103,79 @@ bool FVariantManagerPropertyNode::IsSelectable() const
 	return true;
 }
 
+TOptional<EItemDropZone> FVariantManagerPropertyNode::CanDrop(const FDragDropEvent& DragDropEvent, EItemDropZone ItemDropZone) const
+{
+	TSharedPtr<FVariantManagerDragDropOp> VarManDragDrop = DragDropEvent.GetOperationAs<FVariantManagerDragDropOp>();
+	if (!VarManDragDrop.IsValid())
+	{
+		return TOptional<EItemDropZone>();
+	}
+
+	TSharedPtr<FVariantManager> VarMan = GetVariantManager().Pin();
+	if (!VarMan.IsValid())
+	{
+		return TOptional<EItemDropZone>();
+	}
+
+	TArray<FDisplayNodeRef> PropOrFuncNodes;
+	for (const TSharedRef<FVariantManagerDisplayNode>& DraggedNode : VarManDragDrop->GetDraggedNodes())
+	{
+		if ((DraggedNode->GetType() == EVariantManagerNodeType::Property ||
+			 DraggedNode->GetType() == EVariantManagerNodeType::Function ) &&
+			DraggedNode != SharedThis(this)) // No point in reordering the node ondo itself
+		{
+			PropOrFuncNodes.Add(DraggedNode);
+		}
+	}
+
+	if (PropOrFuncNodes.Num() > 0)
+	{
+		FText NewHoverText = FText::Format( LOCTEXT("CanDrop_Captures", "Reorder '{0}'"), PropOrFuncNodes[0]->GetDisplayName());
+
+		const FSlateBrush* NewHoverIcon = FEditorStyle::GetBrush(TEXT("Graph.ConnectorFeedback.OK"));
+
+		VarManDragDrop->SetToolTip(NewHoverText, NewHoverIcon);
+
+		return ItemDropZone == EItemDropZone::AboveItem ? ItemDropZone : EItemDropZone::BelowItem;
+	}
+
+	if (VarManDragDrop.IsValid())
+	{
+		VarManDragDrop->ResetToDefaultToolTip();
+	}
+
+	return TOptional<EItemDropZone>();
+}
+
+void FVariantManagerPropertyNode::Drop(const FDragDropEvent& DragDropEvent, EItemDropZone ItemDropZone)
+{
+	TSharedPtr<FVariantManager> VarMan = GetVariantManager().Pin();
+	if (!VarMan.IsValid())
+	{
+		return;
+	}
+
+	TSharedPtr<SVariantManager> VarManWidget = VarMan->GetVariantManagerWidget();
+	if (!VarManWidget.IsValid())
+	{
+		return;
+	}
+
+	TSharedPtr<FVariantManagerDragDropOp> VarManDragDrop = DragDropEvent.GetOperationAs<FVariantManagerDragDropOp>();
+	if (!VarManDragDrop.IsValid())
+	{
+		return;
+	}
+
+	TArray<TSharedPtr<FVariantManagerPropertyNode>> DraggedPropertyNodes;
+	for (const TSharedRef<FVariantManagerDisplayNode>& DraggedNode : VarManDragDrop->GetDraggedNodes())
+	{
+		DraggedPropertyNodes.Add(StaticCastSharedRef<FVariantManagerPropertyNode>(DraggedNode));
+	}
+
+	VarManWidget->ReorderPropertyNodes(DraggedPropertyNodes, SharedThis(this), ItemDropZone);
+}
+
 // Without this, SImages (used for example for the browse and useselected buttons next to
 // object property value widgets) will have set SetColorAndOpacity(FSlateColor::UseForeground()).
 // That will cause them to automatically go black when the row is selected, which we don't want.
@@ -426,6 +503,31 @@ TSharedRef<SWidget> FVariantManagerPropertyNode::GetCustomOutlinerContent(TShare
 	];
 }
 
+uint32 FVariantManagerPropertyNode::GetDisplayOrder() const
+{
+	uint32 DisplayOrder = UINT32_MAX;
+	for (const TWeakObjectPtr<UPropertyValue>& PropertyValue : PropertyValues)
+	{
+		if (PropertyValue.IsValid())
+		{
+			DisplayOrder = FMath::Min(DisplayOrder, PropertyValue->GetDisplayOrder());
+		}
+	}
+
+	return DisplayOrder;
+}
+
+void FVariantManagerPropertyNode::SetDisplayOrder(uint32 InDisplayOrder)
+{
+	for (const TWeakObjectPtr<UPropertyValue>& PropertyValue : PropertyValues)
+	{
+		if (PropertyValue.IsValid())
+		{
+			PropertyValue->SetDisplayOrder(InDisplayOrder);
+		}
+	}
+}
+
 void FVariantManagerPropertyNode::UpdateRecordedDataFromSinglePropView(TSharedPtr<ISinglePropertyView> SinglePropView)
 {
 	// Warning: This also fires after UpdateSinglePropViewFromRecordedData when that fires
diff --git a/Engine/Plugins/Enterprise/VariantManager/Source/VariantManager/Private/DisplayNodes/VariantManagerPropertyNode.h b/Engine/Plugins/Enterprise/VariantManager/Source/VariantManager/Private/DisplayNodes/VariantManagerPropertyNode.h
--- a/Engine/Plugins/Enterprise/VariantManager/Source/VariantManager/Private/DisplayNodes/VariantManagerPropertyNode.h
+++ b/Engine/Plugins/Enterprise/VariantManager/Source/VariantManager/Private/DisplayNodes/VariantManagerPropertyNode.h
@@ -41,6 +41,9 @@ public:
 	virtual FText GetDisplayName() const override;
 	virtual void SetDisplayName(const FText& NewDisplayName) override;
 	virtual bool IsSelectable() const override;
+	virtual bool CanDrag() const override { return true; }
+	virtual TOptional<EItemDropZone> CanDrop(const FDragDropEvent& DragDropEvent, EItemDropZone ItemDropZone) const override;
+	virtual void Drop(const FDragDropEvent& DragDropEvent, EItemDropZone ItemDropZone) override;
 	virtual TWeakPtr<FVariantManager> GetVariantManager() const override
 	{
 		return VariantManager;
@@ -49,6 +52,16 @@ public:
 
 	virtual TSharedRef<SWidget> GetCustomOutlinerContent(TSharedPtr<SVariantManagerTableRow> InTableRow) override;
 
+	/**
+	 * Get the order with which the VariantManager should display this in a property list. Lower values will be shown higher up
+	 */
+	virtual uint32 GetDisplayOrder() const;
+
+	/**
+	 * Set the order with which the VariantManager should display this in a property list. Lower values will be shown higher up
+	 */
+	virtual void SetDisplayOrder(uint32 InDisplayOrder);
+
 protected:
 
 	// Callback for when user updates the property widget
diff --git a/Engine/Plugins/Enterprise/VariantManager/Source/VariantManager/Private/SVariantManager.cpp b/Engine/Plugins/Enterprise/VariantManager/Source/VariantManager/Private/SVariantManager.cpp
--- a/Engine/Plugins/Enterprise/VariantManager/Source/VariantManager/Private/SVariantManager.cpp
+++ b/Engine/Plugins/Enterprise/VariantManager/Source/VariantManager/Private/SVariantManager.cpp
@@ -174,6 +174,7 @@ void SVariantManager::Construct(const FArguments& InArgs, TSharedRef<FVariantMan
 	OnObjectPropertyChangedHandle = FCoreUObjectDelegates::OnObjectPropertyChanged.AddRaw(this, &SVariantManager::OnObjectPropertyChanged);
 	OnBeginPieHandle = FEditorDelegates::BeginPIE.AddRaw(this, &SVariantManager::OnPieEvent);
 	OnEndPieHandle = FEditorDelegates::EndPIE.AddRaw(this, &SVariantManager::OnPieEvent);
+	OnEditorSelectionChangedHandle = USelection::SelectionChangedEvent.AddRaw(this, &SVariantManager::OnEditorSelectionChanged);
 
 	// Do this so that if we recompile a function caller changing a function name we'll rebuild our nodes to display the
 	// new names
@@ -388,6 +389,9 @@ SVariantManager::~SVariantManager()
 	FEditorDelegates::EndPIE.Remove(OnEndPieHandle);
 	OnEndPieHandle.Reset();
 
+	USelection::SelectionChangedEvent.Remove(OnEditorSelectionChangedHandle);
+	OnEditorSelectionChangedHandle.Reset();
+
 	if (FModuleManager::Get().IsModuleLoaded(TEXT("LevelEditor")))
 	{
 		FLevelEditorModule& LevelEditorModule = FModuleManager::LoadModuleChecked<FLevelEditorModule>(TEXT("LevelEditor"));
@@ -542,6 +546,18 @@ void SVariantManager::CreateCommandBindings()
 		FCanExecuteAction::CreateSP(this, &SVariantManager::CanAddFunctionCaller)
 	);
 
+	ActorListCommandBindings->MapAction(
+		FVariantManagerEditorCommands::Get().RebindActorDisabled,
+		FExecuteAction::CreateLambda([](){}),
+		FCanExecuteAction::CreateLambda([](){ return false; })
+	);
+
+	ActorListCommandBindings->MapAction(
+		FVariantManagerEditorCommands::Get().RebindToSelected,
+		FExecuteAction::CreateSP(this, &SVariantManager::RebindToSelectedActor),
+		FCanExecuteAction::CreateSP(this, &SVariantManager::CanRebindToSelectedActor)
+	);
+
 	ActorListCommandBindings->MapAction(
 		FVariantManagerEditorCommands::Get().RemoveActorBindings,
 		FExecuteAction::CreateSP(this, &SVariantManager::RemoveActorBindings),
@@ -1310,6 +1326,83 @@ bool SVariantManager::CanAddFunctionCaller()
 	return SelectedActorNodes.Num() > 0;
 }
 
+void SVariantManager::GetSelectedBindingAndEditorActor(UVariantObjectBinding*& OutSelectedBinding, UObject*& OutSelectedObject)
+{
+	TSharedPtr<FVariantManager> PinnedVariantManager = VariantManagerPtr.Pin();
+	if (!PinnedVariantManager.IsValid())
+	{
+		return;
+	}
+
+	FVariantManagerSelection& Selection = PinnedVariantManager->GetSelection();
+	TSet<TSharedRef<FVariantManagerActorNode>>& SelectedActorNodes = Selection.GetSelectedActorNodes();
+	if (SelectedActorNodes.Num() != 1)
+	{
+		return;
+	}
+
+	TSharedRef<FVariantManagerActorNode> SelectedNode = *SelectedActorNodes.CreateIterator();
+	TWeakObjectPtr<UVariantObjectBinding> SelectedBinding = SelectedNode->GetObjectBinding();
+	if (!SelectedBinding.IsValid())
+	{
+		return;
+	}
+	OutSelectedBinding = SelectedBinding.Get();
+
+	// Cache all actors bound to the selected actor node's Variant
+	TSet<UObject*> BoundObjects;
+	UVariant* ParentVariant = SelectedBinding->GetParent();
+	for (const UVariantObjectBinding* Binding : ParentVariant->GetBindings())
+	{
+		BoundObjects.Add(Binding->GetObject());
+	}
+
+	// See if any of the selected actors can be bound to that variant
+	USelection* EditorSelection = GEditor->GetSelectedActors();
+	for (FSelectionIterator SelectedObjectIter(*EditorSelection); SelectedObjectIter; ++SelectedObjectIter)
+	{
+		UObject* EditorSelectedObject = *SelectedObjectIter;
+		if (EditorSelectedObject && !BoundObjects.Contains(EditorSelectedObject))
+		{
+			OutSelectedObject = EditorSelectedObject;
+			return;
+		}
+	}
+}
+
+void SVariantManager::RebindToSelectedActor()
+{
+	UVariantObjectBinding* SelectedBinding = nullptr;
+	UObject* SelectedObject = nullptr;
+	GetSelectedBindingAndEditorActor(SelectedBinding, SelectedObject);
+
+	AActor* SelectedActor = Cast<AActor>(SelectedObject);
+
+	if (SelectedBinding && SelectedActor)
+	{
+		FScopedTransaction Transaction(FText::Format(
+			LOCTEXT("RebindToSelectedTransaction", "Rebind variant to selected actor '{0}'"),
+			FText::FromString(SelectedActor->GetActorLabel())
+		));
+
+		SelectedBinding->SetObject(SelectedObject);
+	}
+}
+
+bool SVariantManager::CanRebindToSelectedActor()
+{
+	UVariantObjectBinding* SelectedBinding = nullptr;
+	UObject* SelectedActor = nullptr;
+	GetSelectedBindingAndEditorActor(SelectedBinding, SelectedActor);
+
+	if (SelectedBinding && SelectedActor)
+	{
+		return true;
+	}
+
+	return false;
+}
+
 void SVariantManager::RemoveActorBindings()
 {
 	TSharedPtr<FVariantManager> PinnedVariantManager = VariantManagerPtr.Pin();
@@ -1751,6 +1844,9 @@ void SVariantManager::OnActorNodeSelectionChanged()
 	TSharedPtr<FVariantManager> PinnedVariantManager = VariantManagerPtr.Pin();
 	if (PinnedVariantManager.IsValid())
 	{
+		// Prevent OnEditorSelectionChanged from firing when we update this selection
+		TGuardValue<bool> Guard(bRespondToEditorSelectionEvents, false);
+
 		const TSet<TSharedRef<FVariantManagerActorNode>>& SelectedActorNodes = PinnedVariantManager->GetSelection().GetSelectedActorNodes();
 
 		GEditor->SelectNone(true, true);
@@ -1952,89 +2048,163 @@ void SVariantManager::RefreshActorList()
 	RefreshPropertyList();
 }
 
+namespace SVariantManagerUtils
+{
+	enum class ERowType : uint8
+	{
+		None = 0,
+		PropertyValue = 1,
+		FunctionCaller = 2,
+	};
+
+	struct FPropertyPanelRow
+	{
+		FPropertyPanelRow(ERowType InRowType, TArray<UPropertyValue*>&& InPropertyValues)
+			: PropertyValues(InPropertyValues)
+			, FunctionCaller(nullptr)
+			, Binding(nullptr)
+			, RowType(InRowType)
+			, DisplayOrder(UINT32_MAX)
+		{
+			ensure(RowType == ERowType::PropertyValue);
+
+			for (UPropertyValue* PropertyValue : PropertyValues)
+			{
+				DisplayOrder = FMath::Min(DisplayOrder, PropertyValue->GetDisplayOrder());
+			}
+		}
+
+		FPropertyPanelRow(ERowType InRowType, FFunctionCaller* InFunctionCaller, UVariantObjectBinding* InBinding)
+			: FunctionCaller(InFunctionCaller)
+			, Binding(InBinding)
+			, RowType(InRowType)
+			, DisplayOrder(UINT32_MAX)
+		{
+			ensure(RowType == ERowType::FunctionCaller);
+
+			DisplayOrder = FunctionCaller->GetDisplayOrder();
+		}
+
+		FName GetName() const
+		{
+			if (RowType == ERowType::PropertyValue && PropertyValues.Num() > 0)
+			{
+				return *PropertyValues[0]->GetFullDisplayString();
+			}
+			else if (RowType == ERowType::FunctionCaller && FunctionCaller)
+			{
+				return FunctionCaller->FunctionName;
+			}
+
+			return NAME_None;
+		}
+
+		TArray<UPropertyValue*> PropertyValues;
+
+		FFunctionCaller* FunctionCaller;
+		UVariantObjectBinding* Binding;
+
+		ERowType RowType;
+		uint32 DisplayOrder;
+	};
+}
+
 void SVariantManager::RefreshPropertyList()
 {
 	FVariantManagerSelection& Selection = VariantManagerPtr.Pin()->GetSelection();
 
-	TArray<UPropertyValue*> NewCapturedProps;
-	TArray<UVariantObjectBinding*> SelectedBindings;
+	TArray<UPropertyValue*> Properties;
+
+	TArray<SVariantManagerUtils::FPropertyPanelRow> Rows;
 
 	for (const TSharedRef<FVariantManagerActorNode>& Node : Selection.GetSelectedActorNodes())
 	{
-		// Ignore unresolved actor bindings
 		UVariantObjectBinding* Binding = Node->GetObjectBinding().Get();
 		if (Binding == nullptr || Binding->GetObject() == nullptr)
 		{
 			continue;
 		}
 
-		NewCapturedProps.Append(Binding->GetCapturedProperties());
-		SelectedBindings.Add(Binding);
-	}
+		Properties.Append(Binding->GetCapturedProperties());
 
-	// Group properties by PathHash
-	TMap<uint32, TArray<UPropertyValue*>> PropsByHash;
-	for (UPropertyValue* NewCapturedProp : NewCapturedProps)
-	{
-		uint32 Hash = NewCapturedProp->GetPropertyPathHash();
-		TArray<UPropertyValue*>& Props = PropsByHash.FindOrAdd(Hash);
-
-		Props.Add(NewCapturedProp);
+		// Add function caller rows first as we need the bindings
+		for (FFunctionCaller& Caller : Binding->GetFunctionCallers())
+		{
+			Rows.Emplace(SVariantManagerUtils::ERowType::FunctionCaller, &Caller, Binding);
+		}
 	}
 
-	DisplayedPropertyNodes.Empty();
-	for (auto HashToPropArray : PropsByHash)
+	// Group properties with the same hash, to show a "multiple object" type of property editor
+	TMap<uint32, TArray<UPropertyValue*>> PropsByHash;
+	for (UPropertyValue* Property : Properties)
 	{
-		TArray<UPropertyValue*>& Props = HashToPropArray.Value;
-		if (Props.Num() < 1)
+		if (!Property)
 		{
 			continue;
 		}
 
-		UPropertyValue* FirstProp = Props[0];
-
-		// Attempts to resolve first so that we can fetch the objects below
-		FirstProp->Resolve();
+		uint32 Hash = Property->GetPropertyPathHash();
+		PropsByHash.FindOrAdd(Hash).Add(Property);
+	}
 
-		UStruct* Struct = FirstProp->GetStructPropertyStruct();
-		UEnum* Enum = FirstProp->GetEnumPropertyEnum();
+	// Add the grouped property rows
+	for (TTuple<uint32, TArray<UPropertyValue*>>& HashToPropArray : PropsByHash)
+	{
+		Rows.Emplace(SVariantManagerUtils::ERowType::PropertyValue, MoveTemp(HashToPropArray.Value));
+	}
+	PropsByHash.Reset();
 
-		if (Struct)
-		{
-			DisplayedPropertyNodes.Add(MakeShareable(new FVariantManagerStructPropertyNode(Props, VariantManagerPtr)));
-		}
-		else if (Enum)
-		{
-			DisplayedPropertyNodes.Add(MakeShareable(new FVariantManagerEnumPropertyNode(Props, VariantManagerPtr)));
-		}
-		else if (FirstProp->GetPropertyClass()->IsChildOf(UStrProperty::StaticClass()) ||
-				 FirstProp->GetPropertyClass()->IsChildOf(UNameProperty::StaticClass()) ||
-				 FirstProp->GetPropertyClass()->IsChildOf(UTextProperty::StaticClass()))
-		{
-			DisplayedPropertyNodes.Add(MakeShareable(new FVariantManagerStringPropertyNode(Props, VariantManagerPtr)));
-		}
-		else if (FirstProp->GetPropCategory() == EPropertyValueCategory::Option)
-		{
-			DisplayedPropertyNodes.Add(MakeShareable(new FVariantManagerOptionPropertyNode(Props, VariantManagerPtr)));
-		}
-		else
+	Rows.Sort([](const SVariantManagerUtils::FPropertyPanelRow& Left, const SVariantManagerUtils::FPropertyPanelRow& Right)
+	{
+		if (Left.DisplayOrder == Right.DisplayOrder)
 		{
-			DisplayedPropertyNodes.Add(MakeShareable(new FVariantManagerPropertyNode(Props, VariantManagerPtr)));
+			return Left.GetName().LexicalLess(Right.GetName());
 		}
-	}
 
-	//TODO @Speed
-	DisplayedPropertyNodes.Sort([](const TSharedPtr<FVariantManagerPropertyNode>& A, const TSharedPtr<FVariantManagerPropertyNode>& B)
-	{
-		return A->GetDisplayName().ToString() < B->GetDisplayName().ToString();
+		return Left.DisplayOrder < Right.DisplayOrder;
 	});
 
-	// Add a node for each function caller
-	for (UVariantObjectBinding* Binding : SelectedBindings)
+	DisplayedPropertyNodes.Empty();
+	for(SVariantManagerUtils::FPropertyPanelRow& Row : Rows)
 	{
-		for (FFunctionCaller& Caller : Binding->GetFunctionCallers())
+		if (Row.RowType == SVariantManagerUtils::ERowType::PropertyValue)
 		{
-			DisplayedPropertyNodes.Add(MakeShareable(new FVariantManagerFunctionPropertyNode(Binding, Caller, VariantManagerPtr)));
+			TArray<UPropertyValue*>& Props = Row.PropertyValues;
+			if (Props.Num() < 1)
+			{
+				continue;
+			}
+
+			// Attempts to resolve first so that we can fetch the objects below
+			UPropertyValue* FirstProp = Props[0];
+			FirstProp->Resolve();
+
+			if (FirstProp->GetStructPropertyStruct())
+			{
+				DisplayedPropertyNodes.Add(MakeShared<FVariantManagerStructPropertyNode>(Props, VariantManagerPtr));
+			}
+			else if (FirstProp->GetEnumPropertyEnum())
+			{
+				DisplayedPropertyNodes.Add(MakeShared<FVariantManagerEnumPropertyNode>(Props, VariantManagerPtr));
+			}
+			else if (FirstProp->GetPropertyClass()->IsChildOf(UStrProperty::StaticClass()) ||
+					 FirstProp->GetPropertyClass()->IsChildOf(UNameProperty::StaticClass()) ||
+					 FirstProp->GetPropertyClass()->IsChildOf(UTextProperty::StaticClass()))
+			{
+				DisplayedPropertyNodes.Add(MakeShared<FVariantManagerStringPropertyNode>(Props, VariantManagerPtr));
+			}
+			else if (FirstProp->GetPropCategory() == EPropertyValueCategory::Option)
+			{
+				DisplayedPropertyNodes.Add(MakeShared<FVariantManagerOptionPropertyNode>(Props, VariantManagerPtr));
+			}
+			else
+			{
+				DisplayedPropertyNodes.Add(MakeShared<FVariantManagerPropertyNode>(Props, VariantManagerPtr));
+			}
+		}
+		else if (Row.RowType == SVariantManagerUtils::ERowType::FunctionCaller && Row.Binding && Row.FunctionCaller)
+		{
+			DisplayedPropertyNodes.Add(MakeShared<FVariantManagerFunctionPropertyNode>(Row.Binding, *Row.FunctionCaller, VariantManagerPtr));
 		}
 	}
 
@@ -2609,5 +2779,88 @@ void SVariantManager::OnPieEvent(bool bIsSimulating)
 	RefreshActorList();
 }
 
+void SVariantManager::ReorderPropertyNodes(const TArray<TSharedPtr<FVariantManagerPropertyNode>>& TheseNodes, TSharedPtr<FVariantManagerPropertyNode> Pivot, EItemDropZone RelativePosition)
+{
+	FScopedTransaction Transaction(LOCTEXT("ReorderedPropertyTransaction", "Reorder property or function captures"));
+
+	// Remove them from their current position
+	for (const TSharedPtr<FVariantManagerPropertyNode>& ThisNode : TheseNodes)
+	{
+		DisplayedPropertyNodes.Remove(ThisNode);
+	}
+
+	// Insert them into their new position
+	int32 Index = DisplayedPropertyNodes.IndexOfByKey(Pivot);
+	Index += RelativePosition == EItemDropZone::BelowItem ? 1 : 0;
+	DisplayedPropertyNodes.Insert(TheseNodes, Index);
+
+	// Normalize the DisplayOrder to an increasing count
+	uint32 Order = 0;
+	for (TSharedPtr<FVariantManagerPropertyNode>& DisplayedNode : DisplayedPropertyNodes)
+	{
+		if (DisplayedNode.IsValid())
+		{
+			DisplayedNode->SetDisplayOrder(Order++);
+		}
+	}
+
+	RefreshPropertyList();
+}
+
+void SVariantManager::OnEditorSelectionChanged(UObject* NewSelection)
+{
+	if (!bRespondToEditorSelectionEvents)
+	{
+		return;
+	}
+
+	USelection* Selection = Cast<USelection>(NewSelection);
+	if (!Selection)
+	{
+		return;
+	}
+
+	TSharedPtr<FVariantManager> PinnedVariantManager = VariantManagerPtr.Pin();
+	if (!PinnedVariantManager.IsValid())
+	{
+		return;
+	}
+
+	FVariantManagerSelection& VariantSelection = PinnedVariantManager->GetSelection();
+
+	TArray<UObject*> SelectedActors;
+	Selection->GetSelectedObjects(AActor::StaticClass(), SelectedActors);
+	TSet<UObject*> SelectedActorsSet{SelectedActors};
+
+	VariantSelection.SuspendBroadcast();
+	for (TSharedRef<FVariantManagerDisplayNode>& DisplayedActor : DisplayedActors)
+	{
+		if (DisplayedActor->GetType() != EVariantManagerNodeType::Actor)
+		{
+			continue;
+		}
+
+		TSharedRef<FVariantManagerActorNode> ActorNode = StaticCastSharedRef<FVariantManagerActorNode>(DisplayedActor);
+		TWeakObjectPtr<UVariantObjectBinding> ObjectBinding = ActorNode->GetObjectBinding();
+		if (!ObjectBinding.IsValid())
+		{
+			continue;
+		}
+
+		UObject* BindingObject = ObjectBinding->GetObject();
+		if (BindingObject && SelectedActorsSet.Contains(BindingObject))
+		{
+			VariantSelection.AddActorNodeToSelection(ActorNode);
+		}
+		else
+		{
+			VariantSelection.RemoveActorNodeFromSelection(ActorNode);
+		}
+	}
+	VariantSelection.ResumeBroadcast();
+
+	RefreshActorList();
+}
+
 #undef LOCTEXT_NAMESPACE
 
diff --git a/Engine/Plugins/Enterprise/VariantManager/Source/VariantManager/Private/SVariantManager.h b/Engine/Plugins/Enterprise/VariantManager/Source/VariantManager/Private/SVariantManager.h
--- a/Engine/Plugins/Enterprise/VariantManager/Source/VariantManager/Private/SVariantManager.h
+++ b/Engine/Plugins/Enterprise/VariantManager/Source/VariantManager/Private/SVariantManager.h
@@ -158,6 +158,9 @@ public:
 	void AddFunctionCaller();
 	bool CanAddFunctionCaller();
 
+	void RebindToSelectedActor();
+	bool CanRebindToSelectedActor();
+
 	void RemoveActorBindings();
 	bool CanRemoveActorBindings();
 
@@ -175,6 +178,8 @@ public:
 
 	void SwitchOnVariant(UVariant* Variant);
 
+	void GetSelectedBindingAndEditorActor(UVariantObjectBinding*& OutSelectedBinding, UObject*& OutSelectedObject);
+
 	// Sorts display nodes based on their order on the screen
 	// Can be used to sort selected nodes
 	void SortDisplayNodes(TArray<TSharedRef<FVariantManagerDisplayNode>>& DisplayNodes);
@@ -223,9 +228,11 @@ public:
 	void OnObjectTransacted(UObject* Object, const class FTransactionObjectEvent& Event);
 	void OnObjectPropertyChanged(UObject* Object, struct FPropertyChangedEvent& Event);
 	void OnPieEvent(bool bIsSimulating);
+	void OnEditorSelectionChanged(UObject* NewSelection);
 
-private:
+	void ReorderPropertyNodes(const TArray<TSharedPtr<FVariantManagerPropertyNode>>& TheseNodes, TSharedPtr<FVariantManagerPropertyNode> Pivot, EItemDropZone RelativePosition);
 
+private:
 	TWeakPtr<FVariantManager> VariantManagerPtr;
 
 	TSharedPtr<SVariantManagerNodeTreeView> NodeTreeView;
@@ -261,10 +268,13 @@ private:
 	FDelegateHandle OnObjectPropertyChangedHandle;
 	FDelegateHandle OnBeginPieHandle;
 	FDelegateHandle OnEndPieHandle;
+	FDelegateHandle OnEditorSelectionChangedHandle;
 
 	// We keep track of this to remember splitter values between loads
 	TSharedPtr<SSplitter> MainSplitter;
-	
+
 	// TODO: Make separate VariantManagerStyle
 	TSharedPtr<FSlateImageBrush> RecordButtonBrush;
+
+	bool bRespondToEditorSelectionEvents = true;
 };
diff --git a/Engine/Plugins/Enterprise/VariantManager/Source/VariantManager/Private/SVariantManagerTableRow.cpp b/Engine/Plugins/Enterprise/VariantManager/Source/VariantManager/Private/SVariantManagerTableRow.cpp
--- a/Engine/Plugins/Enterprise/VariantManager/Source/VariantManager/Private/SVariantManagerTableRow.cpp
+++ b/Engine/Plugins/Enterprise/VariantManager/Source/VariantManager/Private/SVariantManagerTableRow.cpp
@@ -2,14 +2,16 @@
 
 #include "SVariantManagerTableRow.h"
 
-#include "CoreMinimal.h"
-#include "Framework/Commands/GenericCommands.h"
-#include "GameFramework/Actor.h"
+#include "SVariantManager.h"
 #include "VariantManager.h"
 #include "VariantManagerDragDropOp.h"
 #include "VariantManagerSelection.h"
+#include "DisplayNodes/VariantManagerPropertyNode.h"
+
+#include "CoreMinimal.h"
+#include "Framework/Commands/GenericCommands.h"
+#include "GameFramework/Actor.h"
 #include "Styling/SlateIconFinder.h"
-#include "SVariantManager.h"
 #include "Widgets/Views/SHeaderRow.h"
 
 #define LOCTEXT_NAMESPACE "SVariantManagerTableRow"
@@ -37,6 +39,8 @@ FReply SVariantManagerTableRow::DragDetected( const FGeometry& InGeometry, const
 	TSharedPtr<FVariantManagerDisplayNode> PinnedNode = Node.Pin();
 	TSharedPtr<FVariantManager> VarMan = PinnedNode->GetVariantManager().Pin();
 
+	const FSlateBrush* Icon = PinnedNode->GetIconBrush();
+
 	if (PinnedNode.IsValid() && VarMan.IsValid())
 	{
 		FVariantManagerSelection& Selection = VarMan->GetSelection();
@@ -59,14 +63,27 @@ FReply SVariantManagerTableRow::DragDetected( const FGeometry& InGeometry, const
 			}
 
 			int32 NumNodes = DraggableNodes.Num();
+			if (NumNodes > 1)
+			{
+				Icon = FSlateIconFinder::FindIconForClass(AActor::StaticClass()).GetOptionalIcon();
+			}
 
-			DefaultHoverText = FText::Format(NSLOCTEXT("VariantManagerTableRow", "DragActorNode", "{0} actor {0}|plural(one=node,other=nodes)" ),
-				NumNodes);
-
+			DefaultHoverText = FText::Format(NSLOCTEXT("VariantManagerTableRow", "DragActorNode", "{0} actor {0}|plural(one=binding,other=bindings)" ), NumNodes);
 			break;
 		}
+		case EVariantManagerNodeType::Function:
 		case EVariantManagerNodeType::Property:
+		{
+			if (PinnedNode.IsValid() && PinnedNode->CanDrag())
+			{
+				// We only support dragging one node property/function nodes for now
+				ensure(DraggableNodes.Num() == 0);
+
+				DraggableNodes.Add(PinnedNode.ToSharedRef());
+				DefaultHoverText = FText::Format(NSLOCTEXT("VariantManagerTableRow", "DragPropertyNode", "{0}"), PinnedNode->GetDisplayName());
+			}
 			break;
+		}
 		case EVariantManagerNodeType::Variant:  // Intended fallthrough
 		case EVariantManagerNodeType::VariantSet:
 		{
@@ -96,9 +113,6 @@ FReply SVariantManagerTableRow::DragDetected( const FGeometry& InGeometry, const
 
 		VarMan->GetVariantManagerWidget()->SortDisplayNodes(DraggableNodes);
 
-		// TODO: Custom icon depending on dragged node
-		const FSlateBrush* Icon = FSlateIconFinder::FindIconForClass(AActor::StaticClass()).GetOptionalIcon();
-
 		TSharedRef<FVariantManagerDragDropOp> DragDropOp = FVariantManagerDragDropOp::New(DraggableNodes);
 		DragDropOp->SetToolTip(DefaultHoverText, Icon);
 		DragDropOp->SetupDefaults();
diff --git a/Engine/Plugins/Enterprise/VariantManager/Source/VariantManager/Private/VariantManager.cpp b/Engine/Plugins/Enterprise/VariantManager/Source/VariantManager/Private/VariantManager.cpp
--- a/Engine/Plugins/Enterprise/VariantManager/Source/VariantManager/Private/VariantManager.cpp
+++ b/Engine/Plugins/Enterprise/VariantManager/Source/VariantManager/Private/VariantManager.cpp
@@ -724,7 +724,7 @@ TArray<UVariantObjectBinding*> FVariantManager::CreateObjectBindings(const TArra
 			if (UnboundActors.Contains(SelectedActor))
 			{
 				UVariantObjectBinding* NewBinding = NewObject<UVariantObjectBinding>(GetTransientPackage(), NAME_None, RF_Public|RF_Transactional);
-				NewBinding->Init(SelectedActor);
+				NewBinding->SetObject(SelectedActor);
 
 				VarBindingsToTheseActors.Add(NewBinding);
 				NewBindings.Add(NewBinding);
@@ -794,7 +794,7 @@ TArray<UVariantObjectBinding*> FVariantManager::CreateObjectBindingsAndCaptures(
 			if (UnboundActors.Contains(SelectedActor))
 			{
 				UVariantObjectBinding* NewBinding = NewObject<UVariantObjectBinding>(GetTransientPackage(), NAME_None, RF_Public|RF_Transactional);
-				NewBinding->Init(SelectedActor);
+				NewBinding->SetObject(SelectedActor);
 
 				VarBindingsToTheseActors.Add(NewBinding);
 				NewBindings.Add(NewBinding);
diff --git a/Engine/Plugins/Enterprise/VariantManager/Source/VariantManager/Private/VariantManagerEditorCommands.cpp b/Engine/Plugins/Enterprise/VariantManager/Source/VariantManager/Private/VariantManagerEditorCommands.cpp
--- a/Engine/Plugins/Enterprise/VariantManager/Source/VariantManager/Private/VariantManagerEditorCommands.cpp
+++ b/Engine/Plugins/Enterprise/VariantManager/Source/VariantManager/Private/VariantManagerEditorCommands.cpp
@@ -20,6 +20,8 @@ void FVariantManagerEditorCommands::RegisterCommands()
 	UI_COMMAND(ClearThumbnailVariantCommand, "Clear Thumbnail", "Resets the variant thumbnail to default", EUserInterfaceActionType::Button, FInputChord());
 	UI_COMMAND(AddPropertyCaptures, "Add properties", "Capture new properties from the selected actors to this variant", EUserInterfaceActionType::Button, FInputChord());
 	UI_COMMAND(AddFunction, "Add function caller", "Add a function caller to execute blueprint code or trigger events when this variant is switched on", EUserInterfaceActionType::Button, FInputChord());
+	UI_COMMAND(RebindActorDisabled, "Rebind to other Actor", "Actors can only be rebound one at a time!", EUserInterfaceActionType::Button, FInputChord());
+	UI_COMMAND(RebindToSelected, "Bind selected", "Rebind to the first selected actor", EUserInterfaceActionType::Button, FInputChord());
 	UI_COMMAND(RemoveActorBindings, "Remove", "Remove the actor binding from its variant", EUserInterfaceActionType::Button, FInputChord());
 	UI_COMMAND(ApplyProperty, "Apply recorded value", "Applies the recorded value of this property to the bound actor", EUserInterfaceActionType::Button, FInputChord());
 	UI_COMMAND(RecordProperty, "Record current value", "Updates the recorded value with the current value of this property in the bound actor", EUserInterfaceActionType::Button, FInputChord());
diff --git a/Engine/Plugins/Enterprise/VariantManager/Source/VariantManager/Private/VariantManagerEditorCommands.h b/Engine/Plugins/Enterprise/VariantManager/Source/VariantManager/Private/VariantManagerEditorCommands.h
--- a/Engine/Plugins/Enterprise/VariantManager/Source/VariantManager/Private/VariantManagerEditorCommands.h
+++ b/Engine/Plugins/Enterprise/VariantManager/Source/VariantManager/Private/VariantManagerEditorCommands.h
@@ -25,6 +25,8 @@ public:
 	TSharedPtr<FUICommandInfo> ClearThumbnailVariantCommand;
 	TSharedPtr<FUICommandInfo> AddPropertyCaptures;
 	TSharedPtr<FUICommandInfo> AddFunction;
+	TSharedPtr<FUICommandInfo> RebindActorDisabled;
+	TSharedPtr<FUICommandInfo> RebindToSelected;
 	TSharedPtr<FUICommandInfo> RemoveActorBindings;
 	TSharedPtr<FUICommandInfo> ApplyProperty;
 	TSharedPtr<FUICommandInfo> RecordProperty;
diff --git a/Engine/Plugins/Enterprise/VariantManager/Source/VariantManager/Private/VariantManagerSelection.cpp b/Engine/Plugins/Enterprise/VariantManager/Source/VariantManager/Private/VariantManagerSelection.cpp
--- a/Engine/Plugins/Enterprise/VariantManager/Source/VariantManager/Private/VariantManagerSelection.cpp
+++ b/Engine/Plugins/Enterprise/VariantManager/Source/VariantManager/Private/VariantManagerSelection.cpp
@@ -12,12 +12,12 @@ FVariantManagerSelection::FVariantManagerSelection()
 {
 }
 
-const TSet<TSharedRef<FVariantManagerDisplayNode>>& FVariantManagerSelection::GetSelectedOutlinerNodes() const
+TSet<TSharedRef<FVariantManagerDisplayNode>>& FVariantManagerSelection::GetSelectedOutlinerNodes()
 {
 	return SelectedOutlinerNodes;
 }
 
-const TSet<TSharedRef<FVariantManagerActorNode>>& FVariantManagerSelection::GetSelectedActorNodes() const
+TSet<TSharedRef<FVariantManagerActorNode>>& FVariantManagerSelection::GetSelectedActorNodes()
 {
 	return SelectedActorNodes;
 }
diff --git a/Engine/Plugins/Enterprise/VariantManager/Source/VariantManager/Private/VariantManagerSelection.h b/Engine/Plugins/Enterprise/VariantManager/Source/VariantManager/Private/VariantManagerSelection.h
--- a/Engine/Plugins/Enterprise/VariantManager/Source/VariantManager/Private/VariantManagerSelection.h
+++ b/Engine/Plugins/Enterprise/VariantManager/Source/VariantManager/Private/VariantManagerSelection.h
@@ -22,8 +22,8 @@ public:
 	FVariantManagerSelection();
 
 	/** Gets a set of the selected outliner nodes. */
-	const TSet<TSharedRef<FVariantManagerDisplayNode>>& GetSelectedOutlinerNodes() const;
-	const TSet<TSharedRef<FVariantManagerActorNode>>& GetSelectedActorNodes() const;
+	TSet<TSharedRef<FVariantManagerDisplayNode>>& GetSelectedOutlinerNodes();
+	TSet<TSharedRef<FVariantManagerActorNode>>& GetSelectedActorNodes();
 
 	void GetSelectedVariantsAndVariantSets(TArray<UVariant*>& OutVariants, TArray<UVariantSet*>& OutVariantSets);
 
diff --git a/Engine/Plugins/Enterprise/VariantManager/Source/VariantManager/VariantManager.Build.cs b/Engine/Plugins/Enterprise/VariantManager/Source/VariantManager/VariantManager.Build.cs
--- a/Engine/Plugins/Enterprise/VariantManager/Source/VariantManager/VariantManager.Build.cs
+++ b/Engine/Plugins/Enterprise/VariantManager/Source/VariantManager/VariantManager.Build.cs
@@ -22,16 +22,17 @@ namespace UnrealBuildTool.Rules
 			PrivateDependencyModuleNames.AddRange(
 				new string[]
 				{
-					"PropertyEditor",  // For functions that create the property widgets
-					"SlateCore",
-					"Slate",
+                    "BlueprintGraph", // For function director
 					"EditorStyle", // For standard styles on most of UI
-					"InputCore", // For ListView keyboard control
 					"GraphEditor", // For DragDropOp, might be removed later
-                    "BlueprintGraph", // For function director
-					"WorkspaceMenuStructure",
+					"InputCore", // For ListView keyboard control
+					"PropertyEditor",  // For functions that create the property widgets
+					"SceneOutliner",
+					"Slate",
+					"SlateCore",
 					"ToolMenus",
 					"VariantManagerContentEditor",
+					"WorkspaceMenuStructure",
 				}
 			);
         }
diff --git a/Engine/Plugins/Enterprise/VariantManagerContent/Source/VariantManagerContent/Private/FunctionCaller.cpp b/Engine/Plugins/Enterprise/VariantManagerContent/Source/VariantManagerContent/Private/FunctionCaller.cpp
--- a/Engine/Plugins/Enterprise/VariantManagerContent/Source/VariantManagerContent/Private/FunctionCaller.cpp
+++ b/Engine/Plugins/Enterprise/VariantManagerContent/Source/VariantManagerContent/Private/FunctionCaller.cpp
@@ -48,6 +48,16 @@ bool FFunctionCaller::IsValidFunction(UK2Node_FunctionEntry* Function)
 	return PinCategory == UEdGraphSchema_K2::PC_Object || PinCategory == UEdGraphSchema_K2::PC_Interface;
 }
 
+uint32 FFunctionCaller::GetDisplayOrder() const
+{
+	return DisplayOrder;
+}
+
+void FFunctionCaller::SetDisplayOrder(uint32 InDisplayOrder)
+{
+	DisplayOrder = InDisplayOrder;
+}
+
 void FFunctionCaller::CacheFunctionName()
 {
 	UK2Node_FunctionEntry* Node = GetFunctionEntry();
diff --git a/Engine/Plugins/Enterprise/VariantManagerContent/Source/VariantManagerContent/Private/PropertyValue.cpp b/Engine/Plugins/Enterprise/VariantManagerContent/Source/VariantManagerContent/Private/PropertyValue.cpp
--- a/Engine/Plugins/Enterprise/VariantManagerContent/Source/VariantManagerContent/Private/PropertyValue.cpp
+++ b/Engine/Plugins/Enterprise/VariantManagerContent/Source/VariantManagerContent/Private/PropertyValue.cpp
@@ -865,6 +865,8 @@ void UPropertyValue::SetRecordedDataFromEnumIndex(int32 Index)
 
 	ensure(NewValueBytes.Num() == Size);
 
+	Modify();
+
 	SetRecordedDataInternal(NewValueBytes.GetData(), Size);
 }
 
@@ -1077,6 +1079,25 @@ void UPropertyValue::SetRecordedData(const uint8* NewDataBytes, int32 NumBytes,
 	}
 }
 
+#if WITH_EDITORONLY_DATA
+uint32 UPropertyValue::GetDisplayOrder() const
+{
+	return DisplayOrder;
+}
+
+void UPropertyValue::SetDisplayOrder(uint32 InDisplayOrder)
+{
+	if (InDisplayOrder == DisplayOrder)
+	{
+		return;
+	}
+
+	Modify();
+
+	DisplayOrder = InDisplayOrder;
+}
+#endif
+
 void UPropertyValue::SetRecordedDataInternal(const uint8* NewDataBytes, int32 NumBytes, int32 Offset)
 {
 	if (NumBytes <= 0)
diff --git a/Engine/Plugins/Enterprise/VariantManagerContent/Source/VariantManagerContent/Private/VariantManagerObjectVersion.h b/Engine/Plugins/Enterprise/VariantManagerContent/Source/VariantManagerContent/Private/VariantManagerObjectVersion.h
--- a/Engine/Plugins/Enterprise/VariantManagerContent/Source/VariantManagerContent/Private/VariantManagerObjectVersion.h
+++ b/Engine/Plugins/Enterprise/VariantManagerContent/Source/VariantManagerContent/Private/VariantManagerObjectVersion.h
@@ -20,6 +20,8 @@ struct FVariantManagerObjectVersion
 
 		SerializePropertiesAsNames,
 
+		StoreDisplayOrder,
+
 		// -----<new versions can be added above this line>-------------------------------------------------
 		VersionPlusOne,
 		LatestVersion = VersionPlusOne - 1
