Copied cl 10639002

Datasmith: Fix crash on replace reference after a second import

diff --git a/Engine/Plugins/Enterprise/DatasmithImporter/Source/DatasmithImporter/Private/Utility/DatasmithImporterUtils.cpp b/Engine/Plugins/Enterprise/DatasmithImporter/Source/DatasmithImporter/Private/Utility/DatasmithImporterUtils.cpp
--- a/Engine/Plugins/Enterprise/DatasmithImporter/Source/DatasmithImporter/Private/Utility/DatasmithImporterUtils.cpp
+++ b/Engine/Plugins/Enterprise/DatasmithImporter/Source/DatasmithImporter/Private/Utility/DatasmithImporterUtils.cpp
@@ -1116,6 +1116,10 @@ UObject* FDatasmithImporterUtils::StaticDuplicateObject(UObject* SourceObject, U
 		SourceMesh->ClearFlags(RF_ArchetypeObject);
 		DuplicateMesh->ClearFlags(RF_ArchetypeObject);
 
+		// The source mesh is stripped from it's source model, it is not buildable anymore.
+		// -> MarkPendingKill to avoid use-after-move crash in the StaticMesh::Build()
+		SourceMesh->MarkPendingKill();
+
 		for (FStaticMeshSourceModel& SourceModel : SourceModels)
 		{
 			// Fixup the new SourceModels owner
