Copying //UE4/Dev-Enterprise @ cl 9420543 to Dev-Main (//UE4/Dev-Main)

(ushell-p4-cherrypick of 9420574 by JeanMichel.Dignard)

diff --git a/Engine/Plugins/Enterprise/DatasmithImporter/Config/BaseDatasmithImporter.ini b/Engine/Plugins/Enterprise/DatasmithImporter/Config/BaseDatasmithImporter.ini
new file mode 100644
--- /dev/null
+++ b/Engine/Plugins/Enterprise/DatasmithImporter/Config/BaseDatasmithImporter.ini
@@ -0,0 +1,3 @@
+[CoreRedirects]
+
++ClassRedirects=(OldName="/Script/DatasmithImporter.DatasmithBlutilityActor",NewName="/Script/Blutility.PlacedEditorUtilityBase")
diff --git a/Engine/Plugins/Enterprise/DatasmithImporter/Source/DatasmithImporter/Private/DatasmithBlueprintLibrary.cpp b/Engine/Plugins/Enterprise/DatasmithImporter/Source/DatasmithImporter/Private/DatasmithBlueprintLibrary.cpp
new file mode 100644
--- /dev/null
+++ b/Engine/Plugins/Enterprise/DatasmithImporter/Source/DatasmithImporter/Private/DatasmithBlueprintLibrary.cpp
@@ -0,0 +1,467 @@
+// Copyright 1998-2019 Epic Games, Inc. All Rights Reserved.
+
+#include "DatasmithBlueprintLibrary.h"
+
+#include "DatasmithImportFactory.h"
+#include "DatasmithImportOptions.h"
+#include "DatasmithMeshHelper.h"
+#include "DatasmithSceneActor.h"
+#include "DatasmithSceneFactory.h"
+#include "DatasmithSceneSource.h"
+#include "DatasmithStaticMeshImporter.h"
+#include "ObjectElements/DatasmithUSceneElement.h"
+#include "Translators/DatasmithTranslatableSource.h"
+#include "Utility/DatasmithImporterUtils.h"
+
+#include "Async/ParallelFor.h"
+#include "Editor.h"
+#include "Engine/StaticMesh.h"
+#include "Engine/StaticMeshActor.h"
+#include "StaticMeshAttributes.h"
+#include "Engine/World.h"
+#include "GameFramework/Actor.h"
+#include "HAL/FileManager.h"
+#include "Misc/PackageName.h"
+#include "PackageTools.h"
+#include "UObject/Package.h"
+#include "UObject/StrongObjectPtr.h"
+#include "UObject/UObjectGlobals.h"
+#include "Subsystems/AssetEditorSubsystem.h"
+
+#define LOCTEXT_NAMESPACE "DatasmithBlueprintLibrary"
+
+namespace DatasmithStaticMeshBlueprintLibraryUtil
+{
+	void EnsureLightmapUVsAreAvailable( UStaticMesh* StaticMesh )
+	{
+		if ( FMeshDescription* MeshDescription = StaticMesh->GetMeshDescription( 0 ) )
+		{
+			const bool bAreLightmapUVsAvailable = DatasmithMeshHelper::HasUVData( *MeshDescription, StaticMesh->LightMapCoordinateIndex );
+
+			if ( !bAreLightmapUVsAvailable )
+			{
+				FDatasmithStaticMeshImporter::PreBuildStaticMesh( StaticMesh );
+			}
+		}
+	}
+
+	float ParallelogramArea(FVector v0, FVector v1, FVector v2)
+	{
+		FVector TriangleNormal = (v1 - v0) ^ (v2 - v0);
+		float Area = TriangleNormal.Size();
+
+		return Area;
+	}
+}
+
+namespace DatasmithBlueprintLibraryImpl
+{
+	static FName GetLoggerName() { return FName(TEXT("DatasmithLibrary")); }
+
+	static FText GetDisplayName() { return NSLOCTEXT("DatasmithBlueprintLibrary", "LibraryName", "Datasmith Library"); }
+
+	bool ValidatePackage(FString PackageName, UPackage*& OutPackage, const TCHAR* OutFailureReason)
+	{
+		OutFailureReason = TEXT("");
+		OutPackage = nullptr;
+
+		if (PackageName.IsEmpty())
+		{
+			OutFailureReason = TEXT("No destination Folder was provided.");
+			return false;
+		}
+
+		PackageName.ReplaceCharInline(TEXT('\\'), TEXT('/'));
+		while(PackageName.ReplaceInline(TEXT("//"), TEXT("/"), ESearchCase::CaseSensitive)) {}
+		PackageName = UPackageTools::SanitizePackageName(PackageName);
+
+		while (PackageName.EndsWith(TEXT("/")))
+		{
+			PackageName = PackageName.LeftChop(1);
+		}
+
+		FText Reason;
+		if (!FPackageName::IsValidLongPackageName(PackageName, true, &Reason))
+		{
+			OutFailureReason = TEXT("Invalid long package mame.");
+			return false;
+		}
+
+		OutPackage = CreatePackage( nullptr, *PackageName );
+		return OutPackage != nullptr;
+	}
+}
+
+UDatasmithSceneElement* UDatasmithSceneElement::ConstructDatasmithSceneFromFile(const FString& InFilename)
+{
+	using namespace DatasmithBlueprintLibraryImpl;
+	FDatasmithSceneSource Source;
+	Source.SetSourceFile(InFilename);
+
+	UDatasmithSceneElement* DatasmithScene = NewObject<UDatasmithSceneElement>();
+
+	DatasmithScene->SourcePtr.Reset(new FDatasmithTranslatableSceneSource(Source));
+	FDatasmithTranslatableSceneSource& TranslatableSource = *DatasmithScene->SourcePtr;
+
+	if (!TranslatableSource.IsTranslatable())
+	{
+		UE_LOG(LogDatasmithImport, Error, TEXT("Datasmith import error: no suitable translator found for this source. Abort import."));
+		return nullptr;
+	}
+
+	TSharedRef< IDatasmithScene > Scene = FDatasmithSceneFactory::CreateScene(*Source.GetSceneName());
+	DatasmithScene->SetDatasmithSceneElement(Scene);
+
+	bool bLoadConfig = false; //!IsAutomatedImport();
+	DatasmithScene->ImportContextPtr.Reset(new FDatasmithImportContext(Source.GetSourceFile(), bLoadConfig, GetLoggerName(), GetDisplayName(), TranslatableSource.GetTranslator()));
+
+	if (!TranslatableSource.Translate(Scene))
+	{
+		UE_LOG(LogDatasmithImport, Error, TEXT("Datasmith import error: Scene translation failure. Abort import."));
+		return nullptr;
+	}
+
+	return DatasmithScene;
+}
+
+FDatasmithImportFactoryCreateFileResult UDatasmithSceneElement::ImportScene(const FString& DestinationFolder)
+{
+	FDatasmithImportFactoryCreateFileResult Result;
+
+	if (this == nullptr || !ImportContextPtr.IsValid() || !SourcePtr.IsValid() || SourcePtr->GetTranslator() == nullptr || !GetSceneElement().IsValid())
+	{
+		UE_LOG(LogDatasmithImport, Error, TEXT("Invalid State. Ensure ConstructDatasmithSceneFromFile has been called."));
+		return Result;
+	}
+
+	UPackage* DestinationPackage;
+	const TCHAR* OutFailureReason = TEXT("");
+	if (!DatasmithBlueprintLibraryImpl::ValidatePackage(DestinationFolder, DestinationPackage, OutFailureReason))
+	{
+		UE_LOG(LogDatasmithImport, Error, TEXT("Invalid Destination '%s': %s"), *DestinationFolder, OutFailureReason);
+		return Result;
+	}
+
+	FDatasmithImportContext& ImportContext = *ImportContextPtr;
+	TSharedRef< IDatasmithScene > Scene = GetSceneElement().ToSharedRef();
+	EObjectFlags NewObjectFlags = RF_Public | RF_Standalone | RF_Transactional;
+	TSharedPtr<FJsonObject> ImportSettingsJson;
+	bool bIsSilent = true;
+	if ( !ImportContext.Init( Scene, DestinationPackage->GetName(), NewObjectFlags, GWarn, ImportSettingsJson, bIsSilent ) )
+	{
+		return Result;
+	}
+
+	bool bUserCancelled = false;
+	Result.bImportSucceed = DatasmithImportFactoryImpl::ImportDatasmithScene(ImportContext, bUserCancelled);
+	Result.bImportSucceed &= !bUserCancelled;
+
+	if (Result.bImportSucceed)
+	{
+		Result.FillFromImportContext(ImportContext);
+	}
+
+	DestroyScene();
+
+	return Result;
+}
+
+UObject* UDatasmithSceneElement::GetOptions(UClass* OptionType)
+{
+	if (OptionType == nullptr)
+	{
+		OptionType = UDatasmithImportOptions::StaticClass();
+	}
+
+	if (ImportContextPtr.IsValid())
+	{
+		// Standard options from Datasmith
+		if (ImportContextPtr->Options.IsValid() && ImportContextPtr->Options->GetClass()->IsChildOf(OptionType))
+		{
+			return ImportContextPtr->Options.Get();
+		}
+
+		// Additional options from specific translators
+		for(const TStrongObjectPtr<UObject>& AdditionalOption : ImportContextPtr->AdditionalImportOptions)
+		{
+			if (AdditionalOption->GetClass()->IsChildOf(OptionType))
+			{
+				return AdditionalOption.Get();
+			}
+		}
+	}
+	return nullptr;
+}
+
+TMap<UClass*, UObject*> UDatasmithSceneElement::GetAllOptions()
+{
+	TMap<UClass*, UObject*> M;
+
+	auto Append = [&](UObject* Option)
+	{
+		if (Option)
+	{
+			M.Add(Option->GetClass(), Option);
+	}
+	};
+
+	if (ImportContextPtr.IsValid())
+	{
+		// Standard options from Datasmith
+		if (ImportContextPtr->Options.IsValid())
+		{
+			Append(ImportContextPtr->Options.Get());
+		}
+
+		// Additional options from specific translators
+		for(const auto& AdditionalOption : ImportContextPtr->AdditionalImportOptions)
+		{
+			Append(AdditionalOption.Get());
+		}
+	}
+
+	return M;
+}
+
+UDatasmithImportOptions* UDatasmithSceneElement::GetImportOptions()
+{
+	return Cast<UDatasmithImportOptions>(GetOptions());
+}
+
+void UDatasmithSceneElement::DestroyScene()
+{
+	ImportContextPtr.Reset();
+	SourcePtr.Reset();
+	Reset();
+}
+
+void UDatasmithStaticMeshBlueprintLibrary::ComputeLightmapResolution( const TArray< UObject* >& Objects, bool bApplyChanges, float IdealRatio )
+{
+	TRACE_CPUPROFILER_EVENT_SCOPE(UDatasmithStaticMeshBlueprintLibrary::ComputeLightmapResolution)
+
+	// Collect all the static meshes and static mesh components to compute lightmap resolution for
+	TMap< UStaticMesh*, TSet< UStaticMeshComponent* > > WorkingSet;
+
+	for ( UObject* Object : Objects )
+	{
+		if ( AActor* Actor = Cast< AActor >( Object ) )
+		{
+			TInlineComponentArray<UStaticMeshComponent*> StaticMeshComponents(Actor);
+			for (UStaticMeshComponent* StaticMeshComponent : StaticMeshComponents)
+			{
+				if(StaticMeshComponent && StaticMeshComponent->GetStaticMesh())
+				{
+					TSet<UStaticMeshComponent*>& Components = WorkingSet.FindOrAdd( StaticMeshComponent->GetStaticMesh() );
+					Components.Add( StaticMeshComponent );
+				}
+			}
+		}
+		else if ( UStaticMesh* StaticMesh = Cast< UStaticMesh >( Object ) )
+		{
+			TSet<UStaticMeshComponent*>& Components = WorkingSet.FindOrAdd(StaticMesh);
+			Components.Add( nullptr );
+		}
+	}
+
+	// The actual work
+	auto Compute = [&](UStaticMesh* StaticMesh, const TSet<UStaticMeshComponent*>& Components)
+	{
+		bool bComputeForComponents = true;
+
+		// Compute light map resolution for static mesh asset if required
+		if(Components.Contains(nullptr))
+		{
+			if( int32 LightMapResolution = ComputeLightmapResolution(StaticMesh, IdealRatio, FVector::OneVector) )
+			{
+				// Close the mesh editor to prevent crashing. If changes are applied, reopen it after the mesh has been built.
+				bool bStaticMeshIsEdited = false;
+				UAssetEditorSubsystem* AssetEditorSubsystem = GEditor->GetEditorSubsystem<UAssetEditorSubsystem>();
+				if (AssetEditorSubsystem && AssetEditorSubsystem->FindEditorForAsset(StaticMesh, false))
+				{
+					AssetEditorSubsystem->CloseAllEditorsForAsset(StaticMesh);
+					bStaticMeshIsEdited = true;
+				}
+
+				if(bApplyChanges)
+				{
+					StaticMesh->Modify();
+				}
+
+				StaticMesh->LightMapResolution = LightMapResolution;
+
+				if(bApplyChanges)
+				{
+					// Request re-building of mesh with new LODs
+					StaticMesh->PostEditChange();
+
+					// Reopen MeshEditor on this mesh if the MeshEditor was previously opened in it
+					if (bStaticMeshIsEdited)
+					{
+						AssetEditorSubsystem->OpenEditorForAsset(StaticMesh);
+					}
+				}
+			}
+			else
+			{
+				bComputeForComponents = false;
+			}
+		}
+
+		if(bComputeForComponents)
+		{
+			for (UStaticMeshComponent* StaticMeshComponent : Components)
+			{
+				if(StaticMeshComponent)
+				{
+					if( int32 LightMapResolution = ComputeLightmapResolution( StaticMesh, IdealRatio, StaticMeshComponent->GetComponentScale() ) )
+					{
+						StaticMeshComponent->bOverrideLightMapRes = true;
+						StaticMeshComponent->OverriddenLightMapRes = LightMapResolution;
+					}
+				}
+			}
+		}
+	};
+
+	// If no need to notify changes, multi-thread the computing
+	if(!bApplyChanges)
+	{
+		TArray< UStaticMesh* > WorkingSetKeys;
+		WorkingSet.GenerateKeyArray( WorkingSetKeys );
+
+		// Start with the biggest mesh first to help balancing tasks on threads
+		Algo::SortBy(
+			WorkingSetKeys,
+			[](const UStaticMesh* Mesh){ return Mesh->IsMeshDescriptionValid(0) ? Mesh->GetMeshDescription(0)->Vertices().Num() : 0; },
+			TGreater<>()
+		);
+
+		ParallelFor( WorkingSetKeys.Num(),
+			[&]( int32 Index )
+			{
+				Compute(  WorkingSetKeys[Index], WorkingSet[WorkingSetKeys[Index]] );
+			},
+			EParallelForFlags::Unbalanced
+		);
+	}
+	// Do not take any chance, compute sequentially
+	else
+	{
+		for(TPair< UStaticMesh*, TSet< UStaticMeshComponent* > >& Entry : WorkingSet)
+		{
+			Compute(  Entry.Key, Entry.Value );
+		}
+	}
+}
+
+int32 UDatasmithStaticMeshBlueprintLibrary::ComputeLightmapResolution(UStaticMesh* StaticMesh, float IdealRatio, const FVector& StaticMeshScale)
+{
+	if(StaticMesh == nullptr)
+	{
+		return 0;
+	}
+
+	DatasmithStaticMeshBlueprintLibraryUtil::EnsureLightmapUVsAreAvailable( StaticMesh );
+
+	if ( const FMeshDescription* MeshDescription = StaticMesh->GetMeshDescription( 0 ) )
+	{
+		FStaticMeshConstAttributes Attributes(*MeshDescription);
+
+		// Compute the mesh UV density, based on FStaticMeshRenderData::ComputeUVDensities, except that we're working on a MeshDescription
+		const TVertexAttributesConstRef< FVector > VertexPositions = Attributes.GetVertexPositions();
+		const TVertexInstanceAttributesConstRef< FVector2D > VertexUVs = Attributes.GetVertexInstanceUVs();
+
+		if ( VertexUVs.GetNumElements() <= StaticMesh->LightMapCoordinateIndex )
+		{
+			return 0;
+		}
+
+		float MeshArea = 0.f;
+		float MeshUVArea = 0.f;
+
+		TArray< FVector2D > PolygonAreas;
+
+		for ( const FPolygonID PolygonID : MeshDescription->Polygons().GetElementIDs() )
+		{
+			float PolygonArea = 0.f;
+			float PolygonUVArea = 0.f;
+
+			const TArray< FTriangleID >& PolygonTriangleIDs = MeshDescription->GetPolygonTriangleIDs( PolygonID );
+			for ( const FTriangleID TriangleID : PolygonTriangleIDs )
+			{
+				FVector VertexPosition[3];
+				FVector2D LightmapUVs[3];
+
+				for ( int32 CornerIndex = 0; CornerIndex < 3; ++CornerIndex )
+				{
+					const FVertexInstanceID VertexInstanceID = MeshDescription->GetTriangleVertexInstance( TriangleID, CornerIndex );
+					VertexPosition[ CornerIndex ] = VertexPositions[ MeshDescription->GetVertexInstanceVertex( VertexInstanceID ) ] * StaticMeshScale;
+
+					LightmapUVs[ CornerIndex ] = VertexUVs.Get( VertexInstanceID, StaticMesh->LightMapCoordinateIndex );
+				}
+
+				PolygonArea += DatasmithStaticMeshBlueprintLibraryUtil::ParallelogramArea( VertexPosition[0],  VertexPosition[1],  VertexPosition[2] );
+				PolygonUVArea += DatasmithStaticMeshBlueprintLibraryUtil::ParallelogramArea( FVector( LightmapUVs[0], 0.f ),  FVector( LightmapUVs[1], 0.f ),  FVector( LightmapUVs[2], 0.f ) );
+			}
+
+			PolygonAreas.Emplace( FMath::Sqrt( PolygonArea ), FMath::Sqrt(PolygonArea / PolygonUVArea ) );
+		}
+
+		Algo::Sort( PolygonAreas, []( const FVector2D& ElementA, const FVector2D& ElementB )
+		{
+			return ElementA[1] < ElementB[1];
+		} );
+
+		float WeightedUVDensity = 0.f;
+		float Weight = 0.f;
+
+		// Remove 10% of higher and lower texel factors.
+		const int32 Threshold = FMath::FloorToInt( 0.1f * (float)PolygonAreas.Num() );
+		for (int32 Index = Threshold; Index < PolygonAreas.Num() - Threshold; ++Index)
+		{
+			WeightedUVDensity += PolygonAreas[ Index ][ 1 ] * PolygonAreas[ Index ][ 0 ];
+			Weight += PolygonAreas[ Index ][ 0 ];
+		}
+
+		float UVDensity = WeightedUVDensity / Weight;
+
+		int32 LightmapResolution = FMath::FloorToInt( UVDensity * IdealRatio );
+
+		// Ensure that LightmapResolution is a factor of 4
+		return FMath::Max( LightmapResolution + 3 & ~3, 4 );
+	}
+
+	return 0;
+}
+
+FDatasmithImportFactoryCreateFileResult::FDatasmithImportFactoryCreateFileResult()
+	: ImportedBlueprint(nullptr)
+	, bImportSucceed(false)
+{}
+
+void FDatasmithImportFactoryCreateFileResult::FillFromImportContext(const FDatasmithImportContext& ImportContext)
+{
+	switch (ImportContext.Options->HierarchyHandling)
+	{
+	case EDatasmithImportHierarchy::UseMultipleActors:
+		ImportedActors.Append(ImportContext.GetImportedActors());
+		break;
+	case EDatasmithImportHierarchy::UseSingleActor:
+		ImportedActors.Append(ImportContext.ActorsContext.FinalSceneActors.Array());
+		break;
+	case EDatasmithImportHierarchy::UseOneBlueprint:
+		ImportedBlueprint = ImportContext.RootBlueprint;
+		break;
+	default:
+		check(false);
+		break;
+	}
+
+	ImportedMeshes.Reserve(ImportContext.ImportedStaticMeshes.Num());
+	for ( const TPair< TSharedRef< IDatasmithMeshElement >, UStaticMesh* >& MeshPair : ImportContext.ImportedStaticMeshes )
+	{
+		ImportedMeshes.Add( MeshPair.Value );
+	}
+}
+
+#undef LOCTEXT_NAMESPACE
diff --git a/Engine/Plugins/Enterprise/DatasmithImporter/Source/DatasmithImporter/Private/DatasmithConsumer.cpp b/Engine/Plugins/Enterprise/DatasmithImporter/Source/DatasmithImporter/Private/DatasmithConsumer.cpp
new file mode 100644
--- /dev/null
+++ b/Engine/Plugins/Enterprise/DatasmithImporter/Source/DatasmithImporter/Private/DatasmithConsumer.cpp
@@ -0,0 +1,722 @@
+// Copyright 1998-2019 Epic Games, Inc. All Rights Reserved.
+
+#include "DatasmithConsumer.h"
+
+#include "DatasmithActorImporter.h"
+#include "DatasmithAreaLightActor.h"
+#include "DatasmithAssetImportData.h"
+#include "DatasmithAssetUserData.h"
+#include "DatasmithContentEditorModule.h"
+#include "DatasmithContentBlueprintLibrary.h"
+#include "DatasmithScene.h"
+#include "DatasmithSceneActor.h"
+#include "DatasmithSceneFactory.h"
+#include "IDatasmithSceneElements.h"
+#include "DatasmithStaticMeshImporter.h"
+#include "LevelVariantSets.h"
+#include "ObjectTemplates/DatasmithActorTemplate.h"
+#include "ObjectTemplates/DatasmithAreaLightActorTemplate.h"
+#include "ObjectTemplates/DatasmithCineCameraComponentTemplate.h"
+#include "ObjectTemplates/DatasmithLightComponentTemplate.h"
+#include "ObjectTemplates/DatasmithMaterialInstanceTemplate.h"
+#include "ObjectTemplates/DatasmithPointLightComponentTemplate.h"
+#include "ObjectTemplates/DatasmithObjectTemplate.h"
+#include "ObjectTemplates/DatasmithSceneComponentTemplate.h"
+#include "ObjectTemplates/DatasmithStaticMeshTemplate.h"
+#include "ObjectTemplates/DatasmithSkyLightComponentTemplate.h"
+#include "ObjectTemplates/DatasmithStaticMeshComponentTemplate.h"
+#include "Utility/DatasmithImporterUtils.h"
+
+#include "Algo/Count.h"
+#include "AssetRegistryModule.h"
+#include "ComponentReregisterContext.h"
+#include "Components/HierarchicalInstancedStaticMeshComponent.h"
+#include "Components/PointLightComponent.h"
+#include "CineCameraActor.h"
+#include "CineCameraComponent.h"
+#include "EditorLevelUtils.h"
+#include "Engine/Level.h"
+#include "Engine/LevelStreamingDynamic.h"
+#include "Engine/Light.h"
+#include "Engine/StaticMesh.h"
+#include "Engine/StaticMeshActor.h"
+#include "Engine/World.h"
+#include "GameFramework/PhysicsVolume.h"
+#include "GameFramework/WorldSettings.h"
+#include "HAL/FileManager.h"
+#include "IAssetTools.h"
+#include "Internationalization/Internationalization.h"
+#include "LevelSequence.h"
+#include "Materials/Material.h"
+#include "Materials/MaterialInstanceConstant.h"
+#include "Materials/MaterialInterface.h"
+#include "Misc/MessageDialog.h"
+#include "Misc/PackageName.h"
+#include "Misc/Paths.h"
+#include "Misc/SecureHash.h"
+#include "UObject/Package.h"
+#include "UObject/SoftObjectPtr.h"
+
+#define LOCTEXT_NAMESPACE "DatasmithConsumer"
+
+const FText DatasmithConsumerLabel( LOCTEXT( "DatasmithConsumerLabel", "Datasmith writer" ) );
+const FText DatasmithConsumerDescription( LOCTEXT( "DatasmithConsumerDesc", "Writes data prep world's current level and assets to current level" ) );
+
+namespace DatasmithConsumerUtils
+{
+	/** Helper to generate actor element from a scene actor */
+	void ConvertSceneActorsToActors( FDatasmithImportContext& ImportContext );
+
+	/** Helper to pre-build all static meshes from the array of assets passed to a consumer */
+	void AddAssetsToContext( FDatasmithImportContext& ImportContext, TArray< TWeakObjectPtr< UObject > >& Assets );
+
+	FString GetObjectUniqueId(UObject* Object)
+	{
+		UDatasmithContentBlueprintLibrary* DatasmithContentLibrary = Cast< UDatasmithContentBlueprintLibrary >( UDatasmithContentBlueprintLibrary::StaticClass()->GetDefaultObject() );
+		FString DatasmithUniqueId = DatasmithContentLibrary->GetDatasmithUserDataValueForKey( Object, UDatasmithAssetUserData::UniqueIdMetaDataKey );
+
+		return DatasmithUniqueId.IsEmpty() ? Object->GetName() : DatasmithUniqueId;
+	}
+
+	FString GetObjectTag(UObject* Object)
+	{
+		const FString ObjectPath = FPaths::Combine( Object->GetOutermost()->GetName(), Object->GetName() );
+		return FMD5::HashBytes( reinterpret_cast<const uint8*>(*ObjectPath), ObjectPath.Len() * sizeof(TCHAR) );
+	}
+}
+
+bool UDatasmithConsumer::Initialize()
+{
+	FText TaskDescription = LOCTEXT( "DatasmithImportFactory_Initialize", "Preparing world ..." );
+	ProgressTaskPtr = MakeUnique< FDataprepWorkReporter >( Context.ProgressReporterPtr, TaskDescription, 3.0f, 1.0f );
+
+	ProgressTaskPtr->ReportNextStep( LOCTEXT( "DatasmithImportFactory_Initialize", "Preparing world ...") );
+
+	MoveAssets();
+
+	MoveLevel();
+
+	UpdateLevel();
+
+	UPackage* ParentPackage = CreatePackage( nullptr, *TargetContentFolder );
+	ParentPackage->FullyLoad();
+
+	// Check if the Datasmith scene is not already in memory
+	if ( !DatasmithScene.IsValid() )
+	{
+		FName DatasmithSceneName = MakeUniqueObjectName( ParentPackage, UDatasmithScene::StaticClass(), *( GetName() + TEXT("_DS") ) );
+
+		UPackage* Package = CreatePackage( nullptr, *FPaths::Combine( ParentPackage->GetPathName(), DatasmithSceneName.ToString() ) );
+		Package->FullyLoad();
+
+		DatasmithScene = NewObject< UDatasmithScene >( Package, DatasmithSceneName, GetFlags() | RF_Public | RF_Standalone | RF_Transactional );
+		check( DatasmithScene.IsValid() );
+
+		FAssetRegistryModule::AssetCreated( DatasmithScene.Get() );
+		DatasmithScene->MarkPackageDirty();
+
+		DatasmithScene->AssetImportData = NewObject< UDatasmithSceneImportData >( DatasmithScene.Get(), UDatasmithSceneImportData::StaticClass() );
+		check( DatasmithScene->AssetImportData );
+	}
+
+	// #ueent_todo: Find out necessity of namespace for uniqueness of asset's and actor's names
+	if ( !BuildContexts( Context.WorldPtr.Get() ) )
+	{
+		// #ueent_todo: Provide details of why initialization failed
+		return false;
+	}
+
+	// Check if the finalize should be threated as a reimport
+	if (FDatasmithImporterUtils::FindSceneActors( ImportContextPtr->ActorsContext.FinalWorld, ImportContextPtr->SceneAsset).Num() > 0 )
+	{
+		ADatasmithSceneActor* FoundSceneActor = nullptr;
+		for( AActor* Actor : ImportContextPtr->ActorsContext.FinalWorld->GetCurrentLevel()->Actors )
+		{
+			if( ADatasmithSceneActor* SceneActor = Cast<ADatasmithSceneActor>(Actor) )
+			{
+				if( SceneActor->Scene == DatasmithScene )
+				{
+					FoundSceneActor = SceneActor;
+					break;
+				}
+			}
+		}
+
+		if( FoundSceneActor == nullptr )
+		{
+			//Create a new datasmith scene actor in the targeted level
+			FActorSpawnParameters SpawnParameters;
+			SpawnParameters.Template = ImportContextPtr->ActorsContext.ImportSceneActor;
+			ADatasmithSceneActor* DestinationSceneActor = Cast< ADatasmithSceneActor >(ImportContextPtr->ActorsContext.FinalWorld->SpawnActor< ADatasmithSceneActor >(SpawnParameters));
+
+			// Name new destination ADatasmithSceneActor to the DatasmithScene's name
+			DestinationSceneActor->SetActorLabel( ImportContextPtr->Scene->GetName() );
+			DestinationSceneActor->MarkPackageDirty();
+			DestinationSceneActor->RelatedActors.Reset();
+		}
+
+		ImportContextPtr->bIsAReimport = true;
+		ImportContextPtr->Options->ReimportOptions.bRespawnDeletedActors = false;
+		ImportContextPtr->Options->ReimportOptions.bUpdateActors = true;
+		ImportContextPtr->Options->UpdateNotDisplayedConfig( true );
+	}
+
+	return true;
+}
+
+// Inspired from FDataprepDatasmithImporter::FinalizeAssets
+bool UDatasmithConsumer::Run()
+{
+	// Pre-build static meshes
+	ProgressTaskPtr->ReportNextStep( LOCTEXT( "DatasmithImportFactory_PreBuild", "Pre-building assets ...") );
+	FDatasmithStaticMeshImporter::PreBuildStaticMeshes( *ImportContextPtr );
+
+	// No need to have a valid set of assets.
+	// All assets have been added to the AssetContext in UDatasmithConsumer::BuildContexts
+	ProgressTaskPtr->ReportNextStep( LOCTEXT( "DatasmithImportFactory_Finalize", "Finalizing commit ...") );
+	FDatasmithImporter::FinalizeImport( *ImportContextPtr, TSet<UObject*>() );
+
+	// Store package path and level name for subsequent call to Run
+	LastPackagePath = TargetContentFolder;
+	LastLevelName = LevelName;
+
+	return true;
+}
+
+void UDatasmithConsumer::Reset()
+{
+	ImportContextPtr.Reset();
+	ProgressTaskPtr.Reset();
+	UDataprepContentConsumer::Reset();
+
+	// Restore previous current level
+	if( PreviousCurrentLevel != nullptr )
+	{
+		GWorld->SetCurrentLevel( PreviousCurrentLevel );
+		PreviousCurrentLevel = nullptr;
+	}
+}
+
+const FText& UDatasmithConsumer::GetLabel() const
+{
+	return DatasmithConsumerLabel;
+}
+
+const FText& UDatasmithConsumer::GetDescription() const
+{
+	return DatasmithConsumerDescription;
+}
+
+bool UDatasmithConsumer::BuildContexts( UWorld* ImportWorld )
+{
+	UDatasmithSceneImportData* ImportData = Cast< UDatasmithSceneImportData >( DatasmithScene->AssetImportData );
+
+	const FString FilePath = FPaths::Combine( FPaths::ProjectIntermediateDir(), ( DatasmithScene->GetName() + TEXT( ".udatasmith" ) ) );
+
+	ImportContextPtr = MakeUnique< FDatasmithImportContext >( FilePath, false, TEXT("DatasmithImport"), LOCTEXT("DatasmithImportFactoryDescription", "Datasmith") );
+
+	// Update import context with consumer's data
+	ImportContextPtr->Options->BaseOptions = ImportData->BaseOptions;
+	ImportContextPtr->Options->BaseOptions.SceneHandling = EDatasmithImportScene::CurrentLevel;
+	ImportContextPtr->SceneAsset = DatasmithScene.Get();
+	ImportContextPtr->ActorsContext.ImportWorld = ImportWorld;
+	ImportContextPtr->Scene = FDatasmithSceneFactory::CreateScene( *DatasmithScene->GetName() );
+	ImportContextPtr->SceneName = ImportContextPtr->Scene->GetName();
+
+	// Convert all incoming Datasmith scene actors as regular actors
+	DatasmithConsumerUtils::ConvertSceneActorsToActors( *ImportContextPtr );
+
+	// Recreate scene graph from actors in world
+	ImportContextPtr->Scene->SetHost( TEXT( "DatasmithConsumer" ) );
+
+	TArray<AActor*> RootActors;
+	ImportContextPtr->ActorsContext.ImportSceneActor->GetAttachedActors( RootActors );
+	FDatasmithImporterUtils::FillSceneElement( ImportContextPtr->Scene, RootActors );
+
+	// Store IDatasmithScene(Element) in UDatasmithScene
+	FDatasmithImporterUtils::SaveDatasmithScene( ImportContextPtr->Scene.ToSharedRef(), ImportContextPtr->SceneAsset );
+
+	// Initialize context
+	FString SceneOuterPath = DatasmithScene->GetOutermost()->GetName();
+	FString RootPath = FPackageName::GetLongPackagePath( SceneOuterPath );
+
+	if ( Algo::Count( RootPath, TEXT('/') ) > 1 )
+	{
+		// Remove the scene folder as it shouldn't be considered in the import path
+		RootPath.Split( TEXT("/"), &RootPath, nullptr, ESearchCase::IgnoreCase, ESearchDir::FromEnd );
+	}
+
+	FPaths::NormalizeDirectoryName( RootPath );
+
+	if ( !ImportContextPtr->Init( ImportContextPtr->Scene.ToSharedRef(), RootPath, RF_Public | RF_Standalone | RF_Transactional, GWarn, TSharedPtr< FJsonObject >(), true ) )
+	{
+		return false;
+	}
+
+	// Update ImportContext's package data
+	ImportContextPtr->AssetsContext.RootFolderPath = TargetContentFolder;
+	ImportContextPtr->AssetsContext.TransientFolderPath = Context.TransientContentFolder;
+
+	ImportContextPtr->AssetsContext.StaticMeshesFinalPackage.Reset();
+	ImportContextPtr->AssetsContext.MaterialsFinalPackage.Reset();
+	ImportContextPtr->AssetsContext.TexturesFinalPackage.Reset();
+	ImportContextPtr->AssetsContext.LightPackage.Reset();
+	ImportContextPtr->AssetsContext.LevelSequencesFinalPackage.Reset();
+	ImportContextPtr->AssetsContext.LevelVariantSetsFinalPackage.Reset();
+
+	ImportContextPtr->AssetsContext.StaticMeshesImportPackage.Reset();
+	ImportContextPtr->AssetsContext.TexturesImportPackage.Reset();
+	ImportContextPtr->AssetsContext.MaterialsImportPackage.Reset();
+	ImportContextPtr->AssetsContext.MasterMaterialsImportPackage.Reset();
+	ImportContextPtr->AssetsContext.MaterialFunctionsImportPackage.Reset();
+	ImportContextPtr->AssetsContext.LevelSequencesImportPackage.Reset();
+	ImportContextPtr->AssetsContext.LevelVariantSetsImportPackage.Reset();
+
+	// Set the destination world as the one in the level editor
+	ImportContextPtr->ActorsContext.FinalWorld = GWorld;
+
+	// Initialize ActorsContext's UniqueNameProvider with actors in the GWorld not the Import world
+	ImportContextPtr->ActorsContext.UniqueNameProvider = FDatasmithActorUniqueLabelProvider();
+	ImportContextPtr->ActorsContext.UniqueNameProvider.PopulateLabelFrom( GWorld );
+
+	// Copy BaseOptions from ImportContext
+	ImportData->BaseOptions.AssetOptions.PackagePath = ImportContextPtr->Options->BaseOptions.AssetOptions.PackagePath;
+
+	// Add assets as if they have been imported using the current import context
+	DatasmithConsumerUtils::AddAssetsToContext( *ImportContextPtr, Context.Assets );
+
+	return true;
+}
+
+ULevel * UDatasmithConsumer::FindLevel( const FString& InLevelName )
+{
+	UWorld* FinalWorld = GWorld;
+
+	FSoftObjectPath LevelObjectPath( FPaths::Combine( TargetContentFolder, InLevelName ) );
+	UObject* Object = LevelObjectPath.ResolveObject();
+	ULevel* Level = Cast<ULevel>( Object );
+
+	for( const ULevelStreaming* LevelStreaming : FinalWorld->GetStreamingLevels() )
+	{
+		if( LevelStreaming->GetWorldAssetPackageName() == LevelObjectPath.ToString() )
+		{
+			return LevelStreaming->GetLoadedLevel();
+		}
+	}
+
+	return Level;
+}
+
+bool UDatasmithConsumer::SetLevelName( const FString & InLevelName, FText& OutReason )
+{
+	// Update current level
+	if( !InLevelName.IsEmpty() && InLevelName.Compare( TEXT("current"), ESearchCase::IgnoreCase ) != 0 )
+	{
+		// #ueent_todo: What about sub-level of sub-level?
+		if( InLevelName.Contains( TEXT("/") ) || InLevelName.Contains( TEXT("\\") ))
+		{
+			OutReason = LOCTEXT( "DatasmithConsumer_SubLevel", "Sub-level of sub-levels is not supported yet" );
+			return false;
+		}
+
+		// Try to see if there is any issue to eventually create this level, i.e. name collision
+		if( FindLevel( InLevelName ) == nullptr )
+		{
+			FSoftObjectPath LevelObjectPath( FPaths::Combine( TargetContentFolder, InLevelName ) );
+
+			if( StaticFindObject( nullptr, ANY_PACKAGE, *LevelObjectPath.ToString(), true) )
+			{
+				OutReason = LOCTEXT( "DatasmithConsumer_LevelExists", "A object with that name already exists. Please choose another name." );
+				return false;
+			}
+
+			// #ueent_todo: Check if persistent level is locked, etc
+		}
+
+		LevelName = InLevelName;
+	}
+	else
+	{
+		LevelName = TEXT("");
+	}
+
+	OnChanged.Broadcast();
+
+	return true;
+}
+
+void UDatasmithConsumer::MoveAssets()
+{
+	// Do nothing if this is the First call to Run, DatasmithScene is null and LastPackagePath is empty
+	// or the re-Run is using the same package path
+	if( ( !DatasmithScene.IsValid() && LastPackagePath.IsEmpty() ) || LastPackagePath == TargetContentFolder )
+	{
+		return;
+	}
+
+	const FText DialogTitle( LOCTEXT( "DatasmithConsumerDlgTitle", "Warning" ) );
+
+	// Warn user if related Datasmith scene is not in package path and force re-creation of Datasmith scene
+	if( DatasmithScene.IsValid() && !DatasmithScene->GetPathName().StartsWith( TargetContentFolder ) )
+	{
+		FText WarningMessage = FText::Format(LOCTEXT("DatasmithConsumer_NoSceneAsset", "Package path {0} different from path previously used, {1}.\nPrevious content will not be updated."), FText::FromString (TargetContentFolder ), FText::FromString ( LastPackagePath ) );
+		FMessageDialog::Open(EAppMsgType::Ok, WarningMessage, &DialogTitle );
+
+		UE_LOG( LogDatasmithImport, Warning, TEXT("%s"), *WarningMessage.ToString() );
+
+		// Force re-creation of Datasmith scene
+		DatasmithScene.Reset();
+	}
+}
+
+void UDatasmithConsumer::MoveLevel()
+{
+	// Do nothing if this is the First call to Run, DatasmithScene is null and LastLevelName is empty
+	// or the re-Run is using the same level
+	if( ( !DatasmithScene.IsValid() && LastLevelName.IsEmpty() ) || LastLevelName == LevelName )
+	{
+		return;
+	}
+
+	const FText DialogTitle( LOCTEXT( "DatasmithConsumerDlgTitle", "Warning" ) );
+
+	ULevel* Level = FindLevel( LevelName );
+	if( Level == nullptr )
+	{
+		FText WarningMessage = FText::Format(LOCTEXT("DatasmithConsumer_NoLevel", "Level {0} different from level previously used, {1}.\nPrevious level will not be updated."), FText::FromString( LevelName ), FText::FromString (LastLevelName ) );
+		FMessageDialog::Open(EAppMsgType::Ok, WarningMessage, &DialogTitle );
+
+		UE_LOG( LogDatasmithImport, Warning, TEXT("%s"), *WarningMessage.ToString() );
+
+		return;
+	}
+
+	// New level exists, search for DatasmithSceneActor associated with this consumer
+	ADatasmithSceneActor* FoundSceneActor = nullptr;
+	for( AActor* Actor : Level->Actors )
+	{
+		if( ADatasmithSceneActor* SceneActor = Cast<ADatasmithSceneActor>(Actor) )
+		{
+			if( SceneActor->Scene == DatasmithScene.Get() )
+			{
+				FoundSceneActor = SceneActor;
+				break;
+			}
+		}
+	}
+
+	if( FoundSceneActor == nullptr )
+	{
+		FText WarningMessage = FText::Format(LOCTEXT("DatasmithConsumer_NoScene", "Level {0} does not contain main actor from previous execution.\nA new actor will be created."), FText::FromString( LevelName ) );
+		FMessageDialog::Open( EAppMsgType::Ok, WarningMessage, &DialogTitle );
+
+		UE_LOG( LogDatasmithImport, Warning, TEXT("%s"), *WarningMessage.ToString() );
+	}
+}
+
+void UDatasmithConsumer::UpdateLevel()
+{
+	PreviousCurrentLevel = nullptr;
+
+	if( !LevelName.IsEmpty() )
+	{
+		UWorld* FinalWorld = GWorld;
+
+		ULevel* Level = FindLevel( LevelName );
+
+		if( Level == nullptr )
+		{
+			FSoftObjectPath LevelObjectPath( FPaths::Combine( TargetContentFolder, LevelName ) );
+
+			FString PackageFilename;
+			FPackageName::TryConvertLongPackageNameToFilename( LevelObjectPath.ToString(), PackageFilename, FPackageName::GetMapPackageExtension() );
+			if( ULevelStreaming* LevelStreaming = EditorLevelUtils::CreateNewStreamingLevelForWorld( *GWorld, ULevelStreamingDynamic::StaticClass(), *PackageFilename ) )
+			{
+				Level = LevelStreaming->GetLoadedLevel();
+			}
+			else
+			{
+				// #ueent_todo: Warn user that level could not be created
+				Level = FinalWorld->PersistentLevel;
+			}
+
+			check( Level );
+		}
+
+		if( Level != FinalWorld->GetCurrentLevel() )
+		{
+			PreviousCurrentLevel = FinalWorld->GetCurrentLevel();
+			FinalWorld->SetCurrentLevel( Level );
+		}
+	}
+}
+
+namespace DatasmithConsumerUtils
+{
+	void ConvertSceneActorsToActors( FDatasmithImportContext& ImportContext )
+	{
+		UWorld* ImportWorld = ImportContext.ActorsContext.ImportWorld;
+
+		// Find all ADatasmithSceneActor in the world
+		TArray< ADatasmithSceneActor* > SceneActorsToConvert;
+		TArray<AActor*> Actors( ImportWorld->GetCurrentLevel()->Actors );
+		for( AActor* Actor : Actors )
+		{
+			if( ADatasmithSceneActor* ImportSceneActor = Cast<ADatasmithSceneActor>( Actor ) )
+			{
+				SceneActorsToConvert.Add( ImportSceneActor );
+			}
+		}
+
+		// Create the import scene actor for the import context
+		ADatasmithSceneActor* RootSceneActor = FDatasmithImporterUtils::CreateImportSceneActor( ImportContext, FTransform::Identity );
+		if (RootSceneActor == nullptr)
+		{
+			return;
+		}
+		RootSceneActor->Scene = ImportContext.SceneAsset;
+
+		ImportContext.ActorsContext.ImportSceneActor = RootSceneActor;
+
+		// Add existing scene actors as regular actors
+		TMap< FName, TSoftObjectPtr< AActor > >& RelatedActors = RootSceneActor->RelatedActors;
+		RelatedActors.Reserve( ImportWorld->GetCurrentLevel()->Actors.Num() );
+
+		USceneComponent* NewSceneActorRootComponent = RootSceneActor->GetRootComponent();
+		ImportContext.Hierarchy.Push( NewSceneActorRootComponent );
+
+		TArray<AActor*> ActorsToVisit;
+
+		for( ADatasmithSceneActor* SceneActor : SceneActorsToConvert )
+		{
+			// Create AActor to replace scene actor
+			const FString SceneActorName = SceneActor->GetName();
+			const FString SceneActorLabel = SceneActor->GetActorLabel();
+			SceneActor->Rename( nullptr, nullptr, REN_DontCreateRedirectors | REN_NonTransactional );
+
+			// #ueent_todo: is there more to add to the actor element?
+			// Use actor's label instead of name.
+			// Rationale: Datasmith scene actors are created with the same name and label and their name can change when calling SetLabel.
+			TSharedRef< IDatasmithActorElement > RootActorElement = FDatasmithSceneFactory::CreateActor( *SceneActorLabel );
+			RootActorElement->SetLabel( *SceneActorLabel );
+
+			AActor* Actor = FDatasmithActorImporter::ImportBaseActor( ImportContext, RootActorElement );
+			check( Actor && Actor->GetRootComponent());
+
+			FDatasmithImporter::ImportMetaDataForObject( ImportContext, RootActorElement, Actor );
+
+
+			// Copy the transforms
+			USceneComponent* ActorRootComponent = Actor->GetRootComponent();
+			check( ActorRootComponent );
+
+			USceneComponent* SceneActorRootComponent = SceneActor->GetRootComponent();
+
+			ActorRootComponent->SetRelativeTransform( SceneActorRootComponent->GetRelativeTransform() );
+			ActorRootComponent->SetComponentToWorld( SceneActorRootComponent->GetComponentToWorld() );
+
+			// Reparent children of root scene actor to new root actor
+			TArray<USceneComponent*> AttachedChildren;
+			SceneActor->GetRootComponent()->GetChildrenComponents(false, AttachedChildren);
+
+			for(USceneComponent* SceneComponent : AttachedChildren)
+			{
+				SceneComponent->AttachToComponent( ActorRootComponent, FAttachmentTransformRules::KeepRelativeTransform );
+			}
+
+			// Attach new actor to root scene actor
+			ActorRootComponent->AttachToComponent( NewSceneActorRootComponent, FAttachmentTransformRules::KeepRelativeTransform );
+
+			// Delete root scene actor since it is not needed anymore
+			ImportWorld->DestroyActor( SceneActor, false, true );
+			SceneActor->UnregisterAllComponents();
+
+			SceneActor->Rename( nullptr, GetTransientPackage(), REN_DontCreateRedirectors | REN_NonTransactional );
+
+			Actor->RegisterAllComponents();
+
+			// Append children of actor to be later added as related actors
+			TArray<AActor*> Children;
+			Actor->GetAttachedActors( Children );
+
+			ActorsToVisit.Append( Children );
+		}
+
+		// Recursively add all children of previous scene actors as related to new scene actor
+		while ( ActorsToVisit.Num() > 0)
+		{
+			AActor* VisitedActor = ActorsToVisit.Pop();
+			if (VisitedActor == nullptr)
+			{
+				continue;
+			}
+
+			// Add visited actor as actor related to scene actor
+			RelatedActors.Add( FName( *GetObjectUniqueId( VisitedActor ) ), VisitedActor );
+
+			// Continue with children
+			TArray<AActor*> Children;
+			VisitedActor->GetAttachedActors( Children );
+
+			ActorsToVisit.Append( Children );
+		}
+
+		// #ueent_todo: Find a better way to filter those out.
+		auto IsUnregisteredActor = [&](AActor* Actor)
+		{
+			// Skip non-imported actors
+			if( Actor == RootSceneActor || Actor == nullptr || Actor->GetRootComponent() == nullptr || Actor->IsA<AWorldSettings>() || Actor->IsA<APhysicsVolume>() || Actor->IsA<ABrush>() )
+			{
+				return false;
+			}
+
+			// Skip actor which we have already processed
+			return RelatedActors.Find( *GetObjectUniqueId( Actor ) ) ? false : true;
+		};
+
+		// Find remaining root actors (non scene actors)
+		TArray< AActor* > RootActors;
+		for( AActor* Actor : ImportWorld->GetCurrentLevel()->Actors )
+		{
+			if( IsUnregisteredActor( Actor ) )
+			{
+				// Find root actor
+				AActor* RootActor = Actor;
+
+				while( RootActor->GetAttachParentActor() != nullptr )
+				{
+					RootActor = RootActor->GetAttachParentActor();
+				}
+
+				// Attach root actor to root scene actor
+				RootActor->GetRootComponent()->AttachToComponent( NewSceneActorRootComponent, FAttachmentTransformRules::KeepRelativeTransform );
+
+				// Add root actor and its children as related to new scene actor
+				ActorsToVisit.Add( RootActor );
+
+				while ( ActorsToVisit.Num() > 0)
+				{
+					AActor* VisitedActor = ActorsToVisit.Pop();
+					if (VisitedActor == nullptr)
+					{
+						continue;
+					}
+
+					// Add visited actor as actor related to scene actor
+					RelatedActors.Add( FName( *GetObjectUniqueId( VisitedActor ) ), VisitedActor );
+
+					// Continue with children
+					TArray<AActor*> Children;
+					VisitedActor->GetAttachedActors( Children );
+
+					ActorsToVisit.Append( Children );
+				}
+			}
+		}
+	}
+
+	void AddAssetsToContext(FDatasmithImportContext& ImportContext, TArray<TWeakObjectPtr<UObject>>& Assets)
+	{
+		// Addition is done in 2 passes to properly collect UMaterial objects referenced by UMaterialInstance ones
+		// Templates are added to assets which have not been created through Datasmith
+
+		// Add template and Datasmith unique Id to source object
+		auto AddTemplate = [](UClass* TemplateClass, UObject* Source)
+		{
+			UDatasmithObjectTemplate* DatasmithTemplate = NewObject< UDatasmithObjectTemplate >( Source, TemplateClass );
+			DatasmithTemplate->Load( Source );
+			FDatasmithObjectTemplateUtils::SetObjectTemplate( Source, DatasmithTemplate );
+
+			UDatasmithAssetUserData::SetDatasmithUserDataValueForKey(Source, UDatasmithAssetUserData::UniqueIdMetaDataKey, Source->GetName() );
+		};
+
+		// First skip UMaterial objects which are not referenced by a UmaterialInstance one
+		int32 MaterialCount = 0;
+		TSet< UMaterialInterface* > ParentMaterials;
+		for(TWeakObjectPtr<UObject>& AssetPtr : Assets)
+		{
+			if( UObject* Asset = AssetPtr.Get() )
+			{
+				FString AssetTag = DatasmithConsumerUtils::GetObjectTag( Asset );
+
+				if(UTexture* Texture = Cast<UTexture>(Asset))
+				{
+					TSharedRef< IDatasmithTextureElement > TextureElement = FDatasmithSceneFactory::CreateTexture( *AssetTag );
+					ImportContext.ImportedTextures.Add( TextureElement, Texture );
+				}
+				else if(UMaterialInstance* MaterialInstance = Cast<UMaterialInstance>(Asset))
+				{
+					TSharedRef< IDatasmithMaterialElement > MaterialElement = FDatasmithSceneFactory::CreateMaterial( *AssetTag );
+					ImportContext.ImportedMaterials.Add( MaterialElement, MaterialInstance );
+
+					if ( UMaterialInterface* MaterialParent = MaterialInstance->Parent )
+					{
+						FString MaterialInstancePath = MaterialInstance->GetOutermost()->GetName();
+						FString ParentPath = MaterialParent->GetOutermost()->GetName();
+
+						// Add parent material to ImportedParentMaterials if applicable
+						if ( ParentPath.StartsWith( MaterialInstancePath ) )
+						{
+							// #ueent_todo : Do we want to compute the hash of the material and check its existence?
+							ImportContext.ImportedParentMaterials.Add( MaterialCount, MaterialParent );
+							MaterialCount++;
+
+							ParentMaterials.Add( MaterialParent );
+						}
+					}
+
+					if(UMaterialInstanceConstant* MaterialInstanceConstant = Cast<UMaterialInstanceConstant>(MaterialInstance))
+					{
+						if(!FDatasmithObjectTemplateUtils::GetObjectTemplate<UDatasmithMaterialInstanceTemplate>( MaterialInstanceConstant ))
+						{
+							AddTemplate( UDatasmithMaterialInstanceTemplate::StaticClass(), MaterialInstanceConstant );
+						}
+					}
+				}
+				else if(UStaticMesh* StaticMesh = Cast<UStaticMesh>(Asset))
+				{
+					// Clean up static meshes which have incomplete render data.
+					if(StaticMesh->RenderData.IsValid() && !StaticMesh->RenderData->IsInitialized())
+					{
+						StaticMesh->RenderData.Reset();
+					}
+
+					if(FDatasmithObjectTemplateUtils::GetObjectTemplate<UDatasmithStaticMeshTemplate>( StaticMesh ) == nullptr)
+					{
+						AddTemplate( UDatasmithStaticMeshTemplate::StaticClass(), StaticMesh );
+					}
+
+					TSharedRef< IDatasmithMeshElement > MeshElement = FDatasmithSceneFactory::CreateMesh( *AssetTag );
+					ImportContext.ImportedStaticMeshes.Add( MeshElement, StaticMesh );
+				}
+				else if(ULevelSequence* LevelSequence = Cast<ULevelSequence>(Asset))
+				{
+					TSharedRef< IDatasmithLevelSequenceElement > LevelSequenceElement = FDatasmithSceneFactory::CreateLevelSequence( *AssetTag );
+					ImportContext.ImportedLevelSequences.Add( LevelSequenceElement, LevelSequence );
+				}
+				else if(ULevelVariantSets* LevelVariantSets = Cast<ULevelVariantSets>(Asset))
+				{
+					TSharedRef< IDatasmithLevelVariantSetsElement > LevelVariantSetsElement = FDatasmithSceneFactory::CreateLevelVariantSets( *AssetTag );
+					ImportContext.ImportedLevelVariantSets.Add( LevelVariantSetsElement, LevelVariantSets );
+				}
+				// #ueent_todo: Add support for assets which are not of the classes above
+			}
+		}
+
+		// Second take care UMaterial objects which are not referenced by a UmaterialInstance one
+		for( TWeakObjectPtr<UObject>& AssetPtr : Assets )
+		{
+			if( UMaterial* Material = Cast<UMaterial>( AssetPtr.Get() ) )
+			{
+				if( !ParentMaterials.Contains( Material ) )
+				{
+					FString AssetTag = DatasmithConsumerUtils::GetObjectTag( Material );
+					TSharedRef< IDatasmithMaterialElement > MaterialElement = FDatasmithSceneFactory::CreateMaterial( *AssetTag );
+					ImportContext.ImportedMaterials.Add( MaterialElement, Material );
+				}
+			}
+		}
+	}
+}
+
+#undef LOCTEXT_NAMESPACE
diff --git a/Engine/Plugins/Enterprise/DatasmithImporter/Source/DatasmithImporter/Private/DatasmithConsumer.h b/Engine/Plugins/Enterprise/DatasmithImporter/Source/DatasmithImporter/Private/DatasmithConsumer.h
new file mode 100644
--- /dev/null
+++ b/Engine/Plugins/Enterprise/DatasmithImporter/Source/DatasmithImporter/Private/DatasmithConsumer.h
@@ -0,0 +1,73 @@
+// Copyright 1998-2019 Epic Games, Inc. All Rights Reserved.
+
+#pragma once
+
+#include "CoreMinimal.h"
+
+#include "DatasmithImporter.h"
+
+#include "DataPrepContentConsumer.h"
+
+#include "UObject/SoftObjectPath.h"
+
+#include "DatasmithConsumer.generated.h"
+
+class UDatasmithScene;
+class ULevel;
+
+UCLASS(Experimental, config = EditorSettings, HideCategories = (DatasmithConsumerInternal))
+class DATASMITHIMPORTER_API UDatasmithConsumer : public UDataprepContentConsumer
+{
+	GENERATED_BODY()
+
+public:
+	UDatasmithConsumer()
+		: DatasmithScene(nullptr)
+		, PreviousCurrentLevel(nullptr)
+	{
+	}
+
+	UPROPERTY( BlueprintReadOnly, Category = DatasmithConsumerInternal)
+	TSoftObjectPtr<UDatasmithScene> DatasmithScene;
+
+	/** Stores the package path used on the last call to UDatasmithConsumer::Run */
+	UPROPERTY( BlueprintReadOnly, Category = DatasmithConsumerInternal )
+	FString LastPackagePath;
+
+	/** Stores the level used on the last call to UDatasmithConsumer::Run */
+	UPROPERTY( BlueprintReadOnly, Category = DatasmithConsumerInternal )
+	FString LastLevelName;
+
+	// Begin UDataprepContentConsumer overrides
+	virtual const FText& GetLabel() const override;
+	virtual const FText& GetDescription() const override;
+	virtual bool SetLevelName(const FString& InLevelName, FText& OutReason ) override;
+
+protected:
+	virtual bool Initialize() override;
+	virtual bool Run() override;
+	virtual void Reset() override;
+	// End UDataprepContentConsumer overrides
+
+private:
+	/** Temporary code to work with UDataprepContentConsumer */
+	bool BuildContexts( UWorld* ImportWorld );
+
+	/** Returns the level associated to LevelName if it exists */
+	ULevel* FindLevel( const FString& InLevelName );
+
+	/** Move assets if destination package path has changed since last call to UDatasmithConsumer::Run */
+	void MoveAssets();
+
+	/** Move level if destination level's name has changed since last call to UDatasmithConsumer::Run */
+	void MoveLevel();
+
+	/** Set current level to what the user specified */
+	void UpdateLevel();
+
+private:
+	TUniquePtr< FDatasmithImportContext > ImportContextPtr;
+	TUniquePtr< FDataprepWorkReporter > ProgressTaskPtr;
+
+	ULevel* PreviousCurrentLevel;
+};
\ No newline at end of file
diff --git a/Engine/Plugins/Enterprise/DatasmithImporter/Source/DatasmithImporter/Private/DatasmithDataprepFetcher.cpp b/Engine/Plugins/Enterprise/DatasmithImporter/Source/DatasmithImporter/Private/DatasmithDataprepFetcher.cpp
new file mode 100644
--- /dev/null
+++ b/Engine/Plugins/Enterprise/DatasmithImporter/Source/DatasmithImporter/Private/DatasmithDataprepFetcher.cpp
@@ -0,0 +1,24 @@
+// Copyright 1998-2019 Epic Games, Inc. All Rights Reserved.
+
+#include "DatasmithDataprepFetcher.h"
+
+#include "DatasmithContentBlueprintLibrary.h"
+
+/* UDataprepStringMetadataValueFetcher methods
+ *****************************************************************************/
+FString UDatasmithStringMetadataValueFetcher::Fetch_Implementation(const UObject* Object, bool& bOutFetchSucceded) const
+{
+	if ( Object ) 
+	{
+		bOutFetchSucceded = true;
+		return UDatasmithContentBlueprintLibrary::GetDatasmithUserDataValueForKey( const_cast< UObject* >( Object ), Key );
+	}
+
+	bOutFetchSucceded = false;
+	return {};
+}
+
+bool UDatasmithStringMetadataValueFetcher::IsThreadSafe() const
+{
+	return true;
+}
\ No newline at end of file
diff --git a/Engine/Plugins/Enterprise/DatasmithImporter/Source/DatasmithImporter/Private/DatasmithDataprepFetcher.h b/Engine/Plugins/Enterprise/DatasmithImporter/Source/DatasmithImporter/Private/DatasmithDataprepFetcher.h
new file mode 100644
--- /dev/null
+++ b/Engine/Plugins/Enterprise/DatasmithImporter/Source/DatasmithImporter/Private/DatasmithDataprepFetcher.h
@@ -0,0 +1,28 @@
+// Copyright 1998-2019 Epic Games, Inc. All Rights Reserved.
+
+#pragma once
+
+#include "SelectionSystem/DataprepStringFetcher.h"
+
+#include "CoreMinimal.h"
+
+#include "DatasmithDataprepFetcher.generated.h"
+
+UCLASS(BlueprintType, NotBlueprintable, Meta = (DisplayName="Metadata Value", ToolTip="Collect the value for a key from the datasmith user metadata."))
+class UDatasmithStringMetadataValueFetcher final : public UDataprepStringFetcher
+{
+	GENERATED_BODY()
+public:
+	//~ UDataprepStringFetcher interface
+	virtual FString Fetch_Implementation(const UObject* Object, bool& bOutFetchSucceded) const final;
+	//~ End of UDataprepStringFetcher interface
+
+	//~ UDataprepFetcher interface
+	virtual bool IsThreadSafe() const final;
+	//~ End of UDataprepFetcher interface
+
+
+	// The key for the for the string
+	UPROPERTY(EditInstanceOnly, BlueprintReadWrite, Category = "Key")
+	FName Key;
+};
diff --git a/Engine/Plugins/Enterprise/DatasmithImporter/Source/DatasmithImporter/Private/DatasmithDataprepOperation.cpp b/Engine/Plugins/Enterprise/DatasmithImporter/Source/DatasmithImporter/Private/DatasmithDataprepOperation.cpp
new file mode 100644
--- /dev/null
+++ b/Engine/Plugins/Enterprise/DatasmithImporter/Source/DatasmithImporter/Private/DatasmithDataprepOperation.cpp
@@ -0,0 +1,26 @@
+// Copyright 1998-2019 Epic Games, Inc. All Rights Reserved.
+
+#include "DatasmithDataprepOperation.h"
+
+#include "DatasmithBlueprintLibrary.h"
+#include "Utility/DatasmithImporterUtils.h"
+
+#define LOCTEXT_NAMESPACE "DatasmithDataprepOperation"
+
+void UDatasmithComputeLightmapResolutionOperation::OnExecution_Implementation(const FDataprepContext& InContext)
+{
+	// Collect start time to log amount of time spent to import incoming file
+	uint64 StartTime = FPlatformTime::Cycles64();
+	int32 ObjectsCount = InContext.Objects.Num();
+
+	UDatasmithStaticMeshBlueprintLibrary::ComputeLightmapResolution( InContext.Objects, false, IdealRatio );
+
+	// Log time spent to import incoming file in minutes and seconds
+	double ElapsedSeconds = FPlatformTime::ToSeconds64(FPlatformTime::Cycles64() - StartTime);
+
+	int ElapsedMin = int(ElapsedSeconds / 60.0);
+	ElapsedSeconds -= 60.0 * (double)ElapsedMin;
+	UE_LOG( LogDatasmithImport, Log, TEXT("Computation of lightmap resolution of %d object(s) took [%d min %.3f s]"), ObjectsCount, ElapsedMin, ElapsedSeconds );
+}
+
+#undef LOCTEXT_NAMESPACE
diff --git a/Engine/Plugins/Enterprise/DatasmithImporter/Source/DatasmithImporter/Private/DatasmithDataprepOperation.h b/Engine/Plugins/Enterprise/DatasmithImporter/Source/DatasmithImporter/Private/DatasmithDataprepOperation.h
new file mode 100644
--- /dev/null
+++ b/Engine/Plugins/Enterprise/DatasmithImporter/Source/DatasmithImporter/Private/DatasmithDataprepOperation.h
@@ -0,0 +1,35 @@
+// Copyright 1998-2019 Epic Games, Inc. All Rights Reserved.
+
+#pragma once
+
+#include "CoreMinimal.h"
+
+#include "DataPrepOperation.h"
+
+#include "DatasmithDataprepOperation.generated.h"
+
+UCLASS(Experimental, Category = MeshOperation, Meta = (DisplayName="Compute Lightmap Resolution", ToolTip = "For each static mesh to process, recompute the lilghtmap resolution based on the specified ratio") )
+class UDatasmithComputeLightmapResolutionOperation : public UDataprepOperation
+{
+	GENERATED_BODY()
+
+	UDatasmithComputeLightmapResolutionOperation()
+		: IdealRatio( 0.2f )
+	{
+	}
+
+public:
+	UPROPERTY(EditAnywhere, BlueprintReadWrite, Category = LightmapOptions )
+	float IdealRatio;
+
+	//~ Begin UDataprepOperation Interface
+public:
+	virtual FText GetCategory_Implementation() const override
+	{
+		return FDataprepOperationCategories::MeshOperation;
+	}
+
+protected:
+	virtual void OnExecution_Implementation(const FDataprepContext& InContext) override;
+	//~ End UDataprepOperation Interface
+};
diff --git a/Engine/Plugins/Enterprise/DatasmithImporter/Source/DatasmithImporter/Private/DatasmithImportContext.cpp b/Engine/Plugins/Enterprise/DatasmithImporter/Source/DatasmithImporter/Private/DatasmithImportContext.cpp
new file mode 100644
--- /dev/null
+++ b/Engine/Plugins/Enterprise/DatasmithImporter/Source/DatasmithImporter/Private/DatasmithImportContext.cpp
@@ -0,0 +1,525 @@
+// Copyright 1998-2019 Epic Games, Inc. All Rights Reserved.
+
+#include "DatasmithImportContext.h"
+
+#include "DatasmithActorImporter.h"
+#include "DatasmithImportOptions.h"
+#include "DatasmithSceneFactory.h"
+#include "DatasmithUtils.h"
+#include "UI/DatasmithImportOptionsWindow.h"
+#include "Utility/DatasmithImporterUtils.h"
+#include "Translators/DatasmithTranslator.h"
+
+#include "Misc/FeedbackContext.h"
+#include "PackageTools.h"
+#include "Misc/Paths.h"
+#include "Interfaces/IMainFrameModule.h"
+#include "Dom/JsonObject.h"
+#include "EngineUtils.h"
+#include "FileHelpers.h"
+#include "Engine/StaticMesh.h"
+
+
+#define LOCTEXT_NAMESPACE "DatasmithImportPlugin"
+
+const TCHAR* UserOptionPath = TEXT("Unreal Engine/Enterprise/Datasmith/Config/UserOptions.ini");
+
+class FDatasmithImportOptionHelper
+{
+public:
+	/**
+	 * Display options' dialog
+	 *
+	 * @param	ImportOptions	The options to display
+	 * @param	DatasmithScene	The DatasmithScene for which we are displaying the import options
+	 *
+	 * @returns true if user accepted the import.
+	 */
+	static bool DisplayOptionsDialog(const TArray<UObject*>& ImportOptions, TSharedRef< IDatasmithScene > DatasmithScene);
+
+	/** Update options based on content in JSON object*/
+	static void LoadOptions(const TArray<UObject*>& ImportOptions, const TSharedPtr<FJsonObject>& ImportSettingsJson);
+
+	/** Work-around to name uniqueness for UObject class */
+	static void CleanUpOptions(const TArray<UObject*>& ImportOptions);
+};
+
+bool FDatasmithImportOptionHelper::DisplayOptionsDialog(const TArray<UObject*>& ImportOptions, TSharedRef< IDatasmithScene > DatasmithScene)
+{
+	TSharedPtr<SWindow> ParentWindow;
+
+	if (FModuleManager::Get().IsModuleLoaded("MainFrame"))
+	{
+		IMainFrameModule& MainFrame = FModuleManager::LoadModuleChecked<IMainFrameModule>("MainFrame");
+		ParentWindow = MainFrame.GetParentWindow();
+	}
+
+	TSharedRef<SWindow> Window = SNew(SWindow)
+		.Title(LOCTEXT("DatasmithImportSettingsTitle", "Datasmith Import Options"))
+		.SizingRule(ESizingRule::Autosized);
+
+	// First option object is always a UDatasmithImportOptions object
+	UDatasmithImportOptions* MainOptions = Cast<UDatasmithImportOptions>(ImportOptions[0]);
+
+	float SceneVersion;
+	LexFromString( SceneVersion, DatasmithScene->GetExporterVersion() );
+
+	TSharedPtr<SDatasmithOptionsWindow> OptionsWindow;
+	Window->SetContent
+	(
+		SAssignNew(OptionsWindow, SDatasmithOptionsWindow)
+		.ImportOptions(ImportOptions)
+		.WidgetWindow(Window)
+		// note: Spacing in text below is intentional for text alignment
+		.FileNameText(FText::Format(LOCTEXT("DatasmithImportSettingsFileName", "  Import File  :    {0}"), FText::FromString(MainOptions->FileName)))
+		.FilePathText(FText::FromString(MainOptions->FilePath))
+		.FileFormatVersion(SceneVersion)
+		.FileSDKVersion( FText::FromString( DatasmithScene->GetExporterSDKVersion() ) )
+		.PackagePathText(FText::Format(LOCTEXT("DatasmithImportSettingsPackagePath", "  Import To   :    {0}"), FText::FromString(MainOptions->BaseOptions.AssetOptions.PackagePath.ToString())))
+		.ProceedButtonLabel(LOCTEXT("DatasmithOptionWindow_ImportCurLevel", "Import"))
+		.ProceedButtonTooltip(LOCTEXT("DatasmithOptionWindow_ImportCurLevel_ToolTip", "Import the file and add to the current Level"))
+		.CancelButtonLabel(LOCTEXT("DatasmithOptionWindow_Cancel", "Cancel"))
+		.CancelButtonTooltip(LOCTEXT("DatasmithOptionWindow_Cancel_ToolTip", "Cancel importing this file"))
+		.MinDetailHeight(320.f)
+		.MinDetailWidth(450.f)
+	);
+
+	FSlateApplication::Get().AddModalWindow(Window, ParentWindow, false);
+
+	MainOptions->bUseSameOptions = OptionsWindow->UseSameOptions();
+
+	return OptionsWindow->ShouldImport();
+}
+
+void FDatasmithImportOptionHelper::LoadOptions(const TArray<UObject*>& ImportOptions, const TSharedPtr<FJsonObject>& ImportSettingsJson)
+{
+	if (!ImportSettingsJson.IsValid() || ImportSettingsJson->Values.Num() == 0)
+	{
+		return;
+	}
+
+	for (UObject* Object : ImportOptions)
+	{
+		const TSharedPtr<FJsonObject>& OptionDataJsonObject = ImportSettingsJson->GetObjectField(Object->GetName());
+		if (OptionDataJsonObject.IsValid())
+		{
+			for (TFieldIterator<UProperty> It(Object->GetClass()); It; ++It)
+			{
+				if (OptionDataJsonObject->HasField((*It)->GetNameCPP()))
+				{
+					FString PropertyValue = OptionDataJsonObject->GetStringField((*It)->GetNameCPP());
+					uint8* PropertyAddr = (*It)->ContainerPtrToValuePtr<uint8>(Object);
+
+					(*It)->ImportText(*PropertyValue, PropertyAddr, 0, Object);
+				}
+			}
+		}
+	}
+}
+
+void FDatasmithImportOptionHelper::CleanUpOptions(const TArray<UObject*>& ImportOptions)
+{
+	// The names of the option objects is just used for UI. And, the option objects are used once
+	// But the names of UObjects have to be unique and those UObjects are garbaged collected instead of immediately freed.
+	// So, the options are renamed to some generic names which remove any future potential name collision
+	for (UObject* Option : ImportOptions)
+	{
+		Option->Rename(nullptr);
+	}
+}
+
+FDatasmithImportContext::FDatasmithImportContext(const FString& FileName, bool bLoadConfig, const FName& LoggerName, const FText& LoggerLabel, TSharedPtr<IDatasmithTranslator> InSceneTranslator)
+	: SceneTranslator(InSceneTranslator)
+	, Options(NewObject<UDatasmithImportOptions>(GetTransientPackage(), TEXT("Datasmith Import Settings")))
+	, RootBlueprint(nullptr)
+	, SceneAsset(nullptr)
+	, bUserCancelled(false)
+	, bIsAReimport(false)
+	, bImportedViaScript(false)
+	, Warn(nullptr)
+	, AssetsContext(*this)
+	, ContextExtension(nullptr)
+	, Logger(LoggerName, LoggerLabel)
+	, CurrentSceneActorIndex(0)
+	, ReferenceCollector(this)
+{
+	ImportOptions.Add(Options.Get());
+	SetFileName(FileName);
+
+	if (bLoadConfig && Options.IsValid())
+	{
+		FString UserDatasmithOptionsFile = FPaths::Combine(FPlatformProcess::UserSettingsDir(), UserOptionPath);
+		Options->LoadConfig(nullptr, *UserDatasmithOptionsFile);
+	}
+
+	// Force the SceneHandling to be on current level by default.
+	// Note: This is done because this option was previously persisted and can get overwritten
+	Options->BaseOptions.SceneHandling = EDatasmithImportScene::CurrentLevel;
+
+	if (SceneTranslator)
+	{
+		SceneTranslator->GetSceneImportOptions(AdditionalImportOptions);
+		for (const TStrongObjectPtr<UObject>& Option : AdditionalImportOptions)
+		{
+			ImportOptions.Add(Option.Get());
+		}
+	}
+}
+
+void FDatasmithImportContext::SetFileName(const FString& FileName)
+{
+	Options->FileName = FPaths::GetCleanFilename(FileName);
+	Options->FilePath = FPaths::ConvertRelativePathToFull(FileName);
+
+	FileHash = FMD5Hash::HashFile(*Options->FilePath);
+}
+
+void FDatasmithImportContext::AddOption(UObject* InOption, bool bLoadConfig)
+{
+	if (InOption != nullptr)
+	{
+		ImportOptions.Add(InOption);
+		if (bLoadConfig)
+		{
+			FString UserDatasmithOptionsFile = FPaths::Combine(FPlatformProcess::UserSettingsDir(), UserOptionPath);
+			InOption->LoadConfig(nullptr, *UserDatasmithOptionsFile);
+		}
+	}
+}
+
+void FDatasmithImportContext::LogError(const FText& InErrorMessage)
+{
+	Logger.Push(EMessageSeverity::Error, InErrorMessage);
+}
+
+void FDatasmithImportContext::LogWarning(const FText& InWarningMessage, bool bPerformance)
+{
+	Logger.Push(bPerformance ? EMessageSeverity::PerformanceWarning : EMessageSeverity::Warning, InWarningMessage);
+}
+
+void FDatasmithImportContext::LogInfo(const FText& InInfoMessage)
+{
+	Logger.Push(EMessageSeverity::Info, InInfoMessage);
+}
+
+bool FDatasmithImportContext::Init(const FString& InFileName, TSharedRef< IDatasmithScene > InScene, const FString& InImportPath, EObjectFlags InFlags, FFeedbackContext* InWarn, const TSharedPtr<FJsonObject>& ImportSettingsJson, bool bSilent)
+{
+	SetFileName(InFileName);
+	return Init(InScene, InImportPath, InFlags, InWarn, ImportSettingsJson, bSilent);
+}
+
+bool FDatasmithImportContext::Init(TSharedRef< IDatasmithScene > InScene, const FString& InImportPath, EObjectFlags InFlags, FFeedbackContext* InWarn, const TSharedPtr<FJsonObject>& ImportSettingsJson, bool bSilent)
+{
+	check(Options->FileName.Len() > 0);
+	check(Options->FilePath.Len() > 0);
+
+	if (!FModuleManager::Get().IsModuleLoaded("AssetTools"))
+	{
+		UE_LOG(LogDatasmithImport, Warning, TEXT("Import failed. The AssetTools module can't be loaded."));
+		return false;
+	}
+
+	Options->BaseOptions.AssetOptions.PackagePath = FName( *InImportPath );
+
+	if (bSilent)
+	{
+		// Update options based on given JSON object
+		if (ImportSettingsJson.IsValid() && ImportSettingsJson->Values.Num() > 0)
+		{
+			FDatasmithImportOptionHelper::LoadOptions(ImportOptions, ImportSettingsJson);
+		}
+	}
+	else
+	{
+		SetupBaseOptionsVisibility();
+
+		bool bShouldImport = FDatasmithImportOptionHelper::DisplayOptionsDialog(ImportOptions, InScene);
+
+		ResetBaseOptionsVisibility();
+
+		if ( !bShouldImport )
+		{
+			FDatasmithImportOptionHelper::CleanUpOptions(ImportOptions);
+			UE_LOG(LogDatasmithImport, Display, TEXT("Import canceled."));
+			return false;
+		}
+
+		// Update config file with new values
+		FString UserDatasmithOptionsFile = FPaths::Combine(FPlatformProcess::UserSettingsDir(), UserOptionPath);
+		for (UObject* Option : ImportOptions)
+		{
+			Option->SaveConfig(CPF_Config, *UserDatasmithOptionsFile);
+		}
+	}
+
+	if (SceneTranslator)
+	{
+		SceneTranslator->SetSceneImportOptions(AdditionalImportOptions);
+	}
+
+	Options->UpdateNotDisplayedConfig( bIsAReimport );
+
+	FDatasmithImportOptionHelper::CleanUpOptions(ImportOptions);
+
+	if ( !ActorsContext.ImportWorld )
+	{
+		// User is asking to import model in a new level
+		// Check to see if there is nothing to save and act according to user's selection
+		// Note: The code below has been borrowed from UEditorEngine::CreateNewMapForEditing
+		if (Options->BaseOptions.SceneHandling == EDatasmithImportScene::NewLevel)
+		{
+			if (!bSilent)
+			{
+				// Check to see if there are unsaved data and user wants to save them
+				// Import will abort if user selects cancel on Save dialog
+				bool bPromptUserToSave = true; // Ask user if he/she wants to save the unsaved data
+				bool bSaveMapPackages = true;
+				bool bSaveContentPackages = true;
+				if (FEditorFileUtils::SaveDirtyPackages(bPromptUserToSave, bSaveMapPackages, bSaveContentPackages) == false)
+				{
+					// The user pressed cancel. Abort the import so the user doesn't lose any changes
+					return false;
+				}
+			}
+
+			// Force the creation of a new level
+			ActorsContext.ImportWorld = UEditorLoadingAndSavingUtils::NewBlankMap(false);
+		}
+		else if (Options->BaseOptions.SceneHandling == EDatasmithImportScene::CurrentLevel)
+		{
+			ActorsContext.ImportWorld = GWorld;
+			if (ActorsContext.ImportWorld == nullptr)
+			{
+				UE_LOG(LogDatasmithImport, Warning, TEXT("Import failed. There is no World/Map open in the Editor."));
+				return false;
+			}
+		}
+		else
+		{
+			ActorsContext.ImportWorld = nullptr;
+		}
+	}
+
+	Warn = InWarn;
+	Scene = InScene;
+
+	// Initialize the filtered scene as a copy of the original scene. We will use it to then filter out items to import.
+	FilteredScene = FDatasmithSceneFactory::DuplicateScene(Scene.ToSharedRef());
+
+	SceneName = Scene->GetName();
+	bUserCancelled = false;
+
+	ObjectFlags = InFlags | RF_Transactional;
+
+	bool bResult = AssetsContext.Init();
+
+	if ( ShouldImportActors() )
+	{
+		bResult = bResult && ActorsContext.Init();
+	}
+
+	return bResult;
+}
+
+void FDatasmithImportContext::DisplayMessages()
+{
+	Logger.Dump();
+}
+
+void FDatasmithImportContext::SetupBaseOptionsVisibility()
+{
+	if ( bIsAReimport )
+	{
+		UProperty* ReimportOptionsProperty = FindField< UProperty >( Options->GetClass(), GET_MEMBER_NAME_CHECKED( UDatasmithImportOptions, ReimportOptions ) );
+
+		if ( ReimportOptionsProperty )
+		{
+			ReimportOptionsProperty->SetMetaData( TEXT("Category"), TEXT("Reimport") );
+		}
+	}
+}
+
+void FDatasmithImportContext::ResetBaseOptionsVisibility()
+{
+	UProperty* ReimportOptionsProperty = FindField< UProperty >( Options->GetClass(), GET_MEMBER_NAME_CHECKED( UDatasmithImportOptions, ReimportOptions ) );
+
+	if ( ReimportOptionsProperty )
+	{
+		ReimportOptionsProperty->SetMetaData( TEXT("Category"), TEXT("NotVisible") );
+	}
+}
+
+void FDatasmithImportContext::AddImportedActor(AActor* InActor)
+{
+	ImportedActorMap.Add(InActor->GetName(), InActor);
+}
+
+TArray<AActor*> FDatasmithImportContext::GetImportedActors() const
+{
+	TArray<AActor*> Result;
+	if (ImportedActorMap.Num())
+	{
+		Result.Reserve(ImportedActorMap.Num());
+		for (const auto& Itt : ImportedActorMap)
+		{
+			Result.Add(Itt.Value);
+		}
+	}
+
+	return Result;
+}
+
+void FDatasmithImportContext::AddSceneComponent(const FString & InName, USceneComponent* InMeshComponent)
+{
+	ImportedSceneComponentMap.Add(InName, InMeshComponent);
+}
+
+bool FDatasmithImportContext::ShouldImportActors() const
+{
+	return ActorsContext.ImportWorld && Options->BaseOptions.SceneHandling != EDatasmithImportScene::AssetsOnly;
+}
+
+FDatasmithImportContext::FInternalReferenceCollector::FInternalReferenceCollector(FDatasmithImportContext* InImportContext)
+	: ImportContext(InImportContext)
+{
+}
+
+namespace DatasmithImportContextInternal
+{
+	template<typename T>
+	void AddReferenceList(FReferenceCollector& Collector, T& List)
+	{
+		for (auto& Ref : List)
+		{
+			Collector.AddReferencedObject(Ref.Value);
+		}
+	}
+}
+
+void FDatasmithImportContext::FInternalReferenceCollector::AddReferencedObjects(FReferenceCollector& Collector)
+{
+	Collector.AddReferencedObject(ImportContext->ActorsContext.ImportSceneActor);
+	Collector.AddReferencedObject(ImportContext->ActorsContext.ImportWorld);
+	Collector.AddReferencedObject(ImportContext->ActorsContext.FinalWorld);
+
+	Collector.AddReferencedObject(ImportContext->RootBlueprint);
+	Collector.AddReferencedObject(ImportContext->SceneAsset);
+
+	for ( TMap< TSharedRef< IDatasmithMeshElement >, UStaticMesh* >::TIterator It = ImportContext->ImportedStaticMeshes.CreateIterator(); It; ++It )
+	{
+		Collector.AddReferencedObject( It->Value );
+	}
+
+	Collector.AddReferencedObjects(ImportContext->ImportedMaterials);
+	Collector.AddReferencedObjects(ImportContext->ImportedParentMaterials);
+
+	for ( TMap< TSharedRef< IDatasmithLevelSequenceElement >, ULevelSequence* >::TIterator It = ImportContext->ImportedLevelSequences.CreateIterator(); It; ++It )
+	{
+		if ( It->Value )
+		{
+			Collector.AddReferencedObject(It->Value);
+		}
+	}
+
+	for ( TMap< TSharedRef< IDatasmithLevelVariantSetsElement >, ULevelVariantSets* >::TIterator It = ImportContext->ImportedLevelVariantSets.CreateIterator(); It; ++It )
+	{
+		if ( It->Value )
+		{
+			Collector.AddReferencedObject(It->Value);
+		}
+	}
+
+	DatasmithImportContextInternal::AddReferenceList(Collector, ImportContext->ImportedActorMap);
+	DatasmithImportContextInternal::AddReferenceList(Collector, ImportContext->ImportedSceneComponentMap);
+}
+
+FDatasmithActorUniqueLabelProvider::FDatasmithActorUniqueLabelProvider(UWorld* World)
+{
+	PopulateLabelFrom(World);
+}
+
+
+void FDatasmithActorUniqueLabelProvider::PopulateLabelFrom(UWorld* World)
+{
+	if (World)
+	{
+		Clear();
+		for (FActorIterator It(World); It; ++It)
+		{
+			AddExistingName(It->GetActorLabel());
+		}
+	}
+}
+
+FDatasmithActorImportContext::FDatasmithActorImportContext(UWorld* World)
+	: ImportSceneActor(nullptr)
+	, CurrentTargetedScene(nullptr)
+	, ImportWorld(World)
+	, FinalWorld(World)
+{
+}
+
+bool FDatasmithActorImportContext::Init()
+{
+	UniqueNameProvider.PopulateLabelFrom( ImportWorld ); // TODO: Should probably populate based on FinalWorld
+	return true;
+}
+
+FDatasmithAssetsImportContext::FDatasmithAssetsImportContext( FDatasmithImportContext& ImportContext )
+	: ParentContext( ImportContext )
+{
+}
+
+bool FDatasmithAssetsImportContext::Init()
+{
+	FString NewRootFolder;
+	if ( ParentContext.SceneAsset )
+	{
+		NewRootFolder = FPackageName::GetLongPackagePath( ParentContext.SceneAsset->GetOutermost()->GetName() );
+	}
+	else
+	{
+		NewRootFolder = FPaths::Combine( ParentContext.Options->BaseOptions.AssetOptions.PackagePath.ToString(), ParentContext.SceneName );
+	}
+
+	ReInit(NewRootFolder);
+
+	return true;
+}
+
+void FDatasmithAssetsImportContext::ReInit(const FString& NewRootFolder)
+{
+	RootFolderPath = UPackageTools::SanitizePackageName(NewRootFolder);
+
+	StaticMeshesFinalPackage.Reset( CreatePackage( nullptr, *FPaths::Combine( RootFolderPath, TEXT("Geometries") ) ) );
+	MaterialsFinalPackage.Reset( CreatePackage( nullptr, *FPaths::Combine( RootFolderPath, TEXT("Materials") ) ) );
+	TexturesFinalPackage.Reset( CreatePackage( nullptr, *FPaths::Combine( RootFolderPath, TEXT("Textures") ) ) );
+	LightPackage.Reset( CreatePackage( nullptr, *FPaths::Combine( RootFolderPath, TEXT("Lights") ) ) );
+	LevelSequencesFinalPackage.Reset( CreatePackage( nullptr, *FPaths::Combine( RootFolderPath, TEXT("Animations") ) ) );
+	LevelVariantSetsFinalPackage.Reset( CreatePackage( nullptr, *FPaths::Combine( RootFolderPath, TEXT("Variants") ) ) );
+
+	TransientFolderPath = FPaths::Combine( RootFolderPath, TEXT("Temp") );
+
+	StaticMeshesImportPackage.Reset( NewObject< UPackage >( nullptr, *FPaths::Combine( TransientFolderPath, TEXT("Geometries") ), RF_Transient ) );
+	StaticMeshesImportPackage->FullyLoad();
+
+	TexturesImportPackage.Reset( NewObject< UPackage >( nullptr, *FPaths::Combine( TransientFolderPath, TEXT("Textures") ), RF_Transient ) );
+	TexturesImportPackage->FullyLoad();
+
+	MaterialsImportPackage.Reset( NewObject< UPackage >( nullptr, *FPaths::Combine( TransientFolderPath, TEXT("Materials") ), RF_Transient ) );
+	MaterialsImportPackage->FullyLoad();
+
+	MasterMaterialsImportPackage.Reset(NewObject< UPackage >(nullptr, *FPaths::Combine(TransientFolderPath, TEXT("Materials/Master")), RF_Transient));
+	MasterMaterialsImportPackage->FullyLoad();
+
+	MaterialFunctionsImportPackage.Reset( NewObject< UPackage >(nullptr, *FPaths::Combine(TransientFolderPath, TEXT("Materials/Master/Functions")), RF_Transient));
+	MaterialFunctionsImportPackage->FullyLoad();
+
+	LevelSequencesImportPackage.Reset( NewObject< UPackage >( nullptr, *FPaths::Combine( TransientFolderPath, TEXT("Animations") ), RF_Transient ) );
+	LevelSequencesImportPackage->FullyLoad();
+
+	LevelVariantSetsImportPackage.Reset( NewObject< UPackage >( nullptr, *FPaths::Combine( TransientFolderPath, TEXT("Variants") ), RF_Transient ) );
+	LevelVariantSetsImportPackage->FullyLoad();
+}
+
+#undef LOCTEXT_NAMESPACE
