Copying //UE4/Dev-Enterprise @ cl 9420543 to Dev-Main (//UE4/Dev-Main)

diff --git a/Engine/Plugins/Enterprise/DatasmithImporter/Source/DatasmithImporter/Private/DatasmithImporter.cpp b/Engine/Plugins/Enterprise/DatasmithImporter/Source/DatasmithImporter/Private/DatasmithImporter.cpp
new file mode 100644
--- /dev/null
+++ b/Engine/Plugins/Enterprise/DatasmithImporter/Source/DatasmithImporter/Private/DatasmithImporter.cpp
@@ -0,0 +1,2395 @@
+// Copyright 1998-2019 Epic Games, Inc. All Rights Reserved.
+
+#include "DatasmithImporter.h"
+
+#include "DatasmithActorImporter.h"
+#include "DatasmithAdditionalData.h"
+#include "DatasmithAssetImportData.h"
+#include "DatasmithAssetUserData.h"
+#include "DatasmithCameraImporter.h"
+#include "DatasmithImportContext.h"
+#include "DatasmithLevelSequenceImporter.h"
+#include "DatasmithLevelVariantSetsImporter.h"
+#include "DatasmithLightImporter.h"
+#include "DatasmithMaterialImporter.h"
+#include "DatasmithPostProcessImporter.h"
+#include "DatasmithScene.h"
+#include "DatasmithSceneActor.h"
+#include "DatasmithSceneFactory.h"
+#include "DatasmithStaticMeshImporter.h"
+#include "DatasmithTextureImporter.h"
+#include "IDatasmithSceneElements.h"
+#include "DatasmithAnimationElements.h"
+#include "LevelVariantSets.h"
+#include "ObjectTemplates/DatasmithObjectTemplate.h"
+#include "Translators/DatasmithPayload.h"
+#include "Translators/DatasmithTranslator.h"
+#include "Utility/DatasmithImporterUtils.h"
+
+#include "AssetRegistryModule.h"
+#include "AssetToolsModule.h"
+#include "Async/Async.h"
+#include "Async/AsyncWork.h"
+#include "CineCameraComponent.h"
+#include "ComponentReregisterContext.h"
+#include "Components/HierarchicalInstancedStaticMeshComponent.h"
+#include "Containers/Map.h"
+#include "CoreMinimal.h"
+#include "EditorLevelUtils.h"
+#include "Editor/UnrealEdEngine.h"
+#include "Engine/Engine.h"
+#include "Engine/StaticMesh.h"
+#include "Engine/StaticMeshActor.h"
+#include "HAL/FileManager.h"
+#include "IAssetTools.h"
+#include "ISourceControlModule.h"
+#include "ISourceControlProvider.h"
+#include "Kismet2/KismetEditorUtilities.h"
+#include "Landscape.h"
+#include "Layers/LayersSubsystem.h"
+#include "LevelSequence.h"
+#include "MaterialEditingLibrary.h"
+#include "MaterialShared.h"
+#include "Materials/MaterialFunction.h"
+#include "Materials/Material.h"
+#include "Materials/MaterialInstanceConstant.h"
+#include "Misc/FeedbackContext.h"
+#include "Misc/FileHelper.h"
+#include "Misc/ScopedSlowTask.h"
+#include "ObjectTools.h"
+#include "PackageTools.h"
+#include "Serialization/ArchiveReplaceObjectRef.h"
+#include "Serialization/MemoryReader.h"
+#include "Serialization/ObjectReader.h"
+#include "Serialization/ObjectWriter.h"
+#include "Settings/EditorExperimentalSettings.h"
+#include "SourceControlOperations.h"
+#include "UnrealEdGlobals.h"
+#include "UObject/Package.h"
+
+extern UNREALED_API UEditorEngine* GEditor;
+
+#define LOCTEXT_NAMESPACE "DatasmithImporter"
+
+namespace DatasmithImporterImpl
+{
+	static void SafetyGarbageCollect()
+	{
+		const int SkipCountThreshold = 1000; // Don't try to cleanup to frequently
+
+		static int SkippedCount = 0;
+		if (++SkippedCount > SkipCountThreshold)
+		{
+			CollectGarbage(GARBAGE_COLLECTION_KEEPFLAGS);
+			SkippedCount = 0;
+		}
+	}
+
+	static UObject* PublicizeAsset( UObject* SourceAsset, const TCHAR* DestinationPath, UObject* ExistingAsset )
+	{
+		TRACE_CPUPROFILER_EVENT_SCOPE(DatasmithImporterImpl::PublicizeAsset);
+
+		SafetyGarbageCollect();
+
+		UPackage* DestinationPackage;
+
+		if ( !ExistingAsset )
+		{
+			FString DestinationPackagePath = UPackageTools::SanitizePackageName( FPaths::Combine( DestinationPath, SourceAsset->GetName() ) );
+			FString DestinationAssetPath = DestinationPackagePath + TEXT(".") + UPackageTools::SanitizePackageName( SourceAsset->GetName() );
+
+			ExistingAsset = FDatasmithImporterUtils::FindObject<UObject>( nullptr, DestinationAssetPath );
+
+			DestinationPackage = ExistingAsset ? ExistingAsset->GetOutermost() : CreatePackage( nullptr, *DestinationPackagePath );
+		}
+		else
+		{
+			DestinationPackage = ExistingAsset->GetOutermost();
+		}
+
+		DestinationPackage->FullyLoad();
+
+		UObject* DestinationAsset = ExistingAsset;
+
+		FString OldAssetPathName;
+
+		// If the object already exist, then we need to fix up the reference
+		if ( ExistingAsset != nullptr && ExistingAsset != SourceAsset )
+		{
+			OldAssetPathName = ExistingAsset->GetPathName();
+
+			DestinationAsset = DuplicateObject< UObject >( SourceAsset, DestinationPackage, ExistingAsset->GetFName() );
+
+			// If mesh's label has changed, update its name
+			if ( ExistingAsset->GetFName() != SourceAsset->GetFName() )
+			{
+				DestinationAsset->Rename( *SourceAsset->GetName(), DestinationPackage, REN_DontCreateRedirectors | REN_NonTransactional );
+			}
+
+			if ( UStaticMesh* DestinationMesh = Cast< UStaticMesh >( DestinationAsset ) )
+			{
+				// This is done during the static mesh build process but we need to redo it after the DuplicateObject since the links are now valid
+				for ( TObjectIterator< UStaticMeshComponent > It; It; ++It )
+				{
+					if ( It->GetStaticMesh() == DestinationMesh )
+					{
+						It->FixupOverrideColorsIfNecessary( true );
+						It->InvalidateLightingCache();
+					}
+				}
+			}
+		}
+		else
+		{
+			SourceAsset->Rename( *SourceAsset->GetName(), DestinationPackage, REN_DontCreateRedirectors | REN_NonTransactional );
+			DestinationAsset = SourceAsset;
+		}
+
+		DestinationAsset->SetFlags( RF_Public );
+		DestinationAsset->MarkPackageDirty();
+
+		if ( !ExistingAsset )
+		{
+			FAssetRegistryModule::AssetCreated( DestinationAsset );
+		}
+		else if ( !OldAssetPathName.IsEmpty() )
+		{
+			FAssetRegistryModule::AssetRenamed( DestinationAsset, OldAssetPathName );
+		}
+
+		return DestinationAsset;
+	}
+
+	/**
+	 * Verifies the input asset can successfully be saved and/or cooked.
+	 * The code below is heavily inspired from:
+	 *			- ContentBrowserUtils::IsValidPackageForCooking
+	 * #ueent_todo Make ContentBrowserUtils::IsValidPackageForCooking public. Not game to do it so late in 4.22
+	 */
+	static void CheckAssetPersistenceValidity(UObject* Asset, FDatasmithImportContext& ImportContext)
+	{
+		if (Asset == nullptr)
+		{
+			return;
+		}
+
+		UPackage* Package = Asset->GetOutermost();
+		const FString PackageName = Package->GetName();
+
+		// Check that package can be saved
+		const FString BasePackageFileName = FPackageName::LongPackageNameToFilename( PackageName );
+		const FString AbsolutePathToAsset = FPaths::ConvertRelativePathToFull( BasePackageFileName );
+
+		// Create fake filename of same length of final asset file name to test ability to write
+		const FString FakeAbsolutePathToAsset = AbsolutePathToAsset + TEXT( ".uasset" );
+
+		// Verify asset file name does not exceed OS' maximum path length
+		if( FPlatformMisc::GetMaxPathLength() < FakeAbsolutePathToAsset.Len() )
+		{
+			ImportContext.LogWarning(FText::Format(LOCTEXT("DatasmithImportInvalidLength", "Saving may partially fail because path for asset {0} is too long. Rename before saving."), FText::FromString( PackageName )));
+		}
+		// Verify user can overwrite existing file
+		else if( IFileManager::Get().FileExists( *FakeAbsolutePathToAsset ) )
+		{
+			FFileStatData FileStatData = IFileManager::Get().GetStatData( *FakeAbsolutePathToAsset );
+			if ( FileStatData.bIsReadOnly )
+			{
+				// Check to see if the file is not under source control
+				bool bWarnUser = true;
+
+				ISourceControlProvider& SourceControlProvider = ISourceControlModule::Get().GetProvider();
+				if (SourceControlProvider.IsAvailable() && SourceControlProvider.IsEnabled())
+				{
+					SourceControlProvider.Execute(ISourceControlOperation::Create<FUpdateStatus>(), FakeAbsolutePathToAsset);
+					FSourceControlStatePtr SourceControlState = SourceControlProvider.GetState(FakeAbsolutePathToAsset, EStateCacheUsage::Use);
+					if( SourceControlState.IsValid() && SourceControlState->CanCheckout() )
+					{
+						// User will be prompted to check out this file when he/she saves the asset. No need to warn.
+						bWarnUser = false;
+					}
+				}
+
+				if(bWarnUser)
+				{
+					ImportContext.LogWarning(FText::Format(LOCTEXT("DatasmithImportInvalidSaving", "Saving may partially fail because file asset {0} cannot be overwritten. Check your privileges."), FText::FromString( PackageName )));
+				}
+			}
+		}
+		// Verify user has privileges to write in folder where asset file will be stored
+		else
+		{
+			// We can't just check for the target content folders with IFileManager::Get().GetStatData here as those will
+			// only be created when UUnrealEdEngine::GetWarningStateForWritePermission is called to check for
+			// write permissions the first time, as the result is cached in GUnrealEd->PackagesCheckedForWritePermission.
+			// To check for permission, we need to first check this cache, and if the PackageName hasn't been
+			// checked yet, we need to replicate what UUnrealEdEngine::GetWarningStateForWritePermission does
+			EWriteDisallowedWarningState WarningState = EWriteDisallowedWarningState::WDWS_MAX;
+			if (GUnrealEd != nullptr && GUnrealEd->PackagesCheckedForWritePermission.Find(PackageName))
+			{
+				WarningState = (EWriteDisallowedWarningState)*GUnrealEd->PackagesCheckedForWritePermission.Find(PackageName);
+			}
+			else if (FFileHelper::SaveStringToFile(TEXT("Write Test"), *FakeAbsolutePathToAsset))
+			{
+				// We can successfully write to the folder containing the package.
+				// Delete the temp file.
+				IFileManager::Get().Delete(*FakeAbsolutePathToAsset);
+				WarningState = EWriteDisallowedWarningState::WDWS_WarningUnnecessary;
+			}
+
+			if(WarningState != EWriteDisallowedWarningState::WDWS_WarningUnnecessary)
+			{
+				ImportContext.LogWarning(FText::Format(LOCTEXT("DatasmithImportInvalidFolder", "Cannot write in folder {0} to store asset {1}. Check access to folder."), FText::FromString( FPaths::GetPath( FakeAbsolutePathToAsset ) ), FText::FromString( PackageName )));
+			}
+		}
+
+		// Check that package can be cooked
+		// Value for MaxGameNameLen directly taken from ContentBrowserUtils::GetPackageLengthForCooking
+		constexpr int32 MaxGameNameLen = 20;
+
+		// Pad out the game name to the maximum allowed
+		const FString GameName = FApp::GetProjectName();
+		FString GameNamePadded = GameName;
+		while (GameNamePadded.Len() < MaxGameNameLen)
+		{
+			GameNamePadded += TEXT(" ");
+		}
+
+		const FString AbsoluteGamePath = FPaths::ConvertRelativePathToFull(FPaths::ProjectDir());
+		const FString AbsoluteGameCookPath = AbsoluteGamePath / TEXT("Saved") / TEXT("Cooked") / TEXT("WindowsNoEditor") / GameName;
+
+		FString AssetPathWithinCookDir = AbsolutePathToAsset;
+		FPaths::RemoveDuplicateSlashes(AssetPathWithinCookDir);
+		AssetPathWithinCookDir.RemoveFromStart(AbsoluteGamePath, ESearchCase::CaseSensitive);
+
+		// Test that the package can be cooked based on the current project path
+		FString AbsoluteCookPathToAsset = AbsoluteGameCookPath / AssetPathWithinCookDir;
+
+		AbsoluteCookPathToAsset.ReplaceInline(*GameName, *GameNamePadded, ESearchCase::CaseSensitive);
+
+		// Get the longest path allowed by the system or use 260 as the longest which is the shortest max path of any platforms that support cooking
+		const int32 MaxCookPath = GetDefault<UEditorExperimentalSettings>()->bEnableLongPathsSupport ? FPlatformMisc::GetMaxPathLength() : 260 /* MAX_PATH */;
+
+		if (AbsoluteCookPathToAsset.Len() > MaxCookPath)
+		{
+			ImportContext.LogWarning(FText::Format(LOCTEXT("DatasmithImportInvalidCooking", "Cooking may fail because path for asset {0} is too long. Rename before cooking."), FText::FromString(PackageName)));
+		}
+	}
+
+	/** Set the texture mode on each texture element based on its usage in the materials */
+	static void SetTexturesMode( FDatasmithImportContext& ImportContext )
+	{
+		const int32 TexturesCount = ImportContext.FilteredScene->GetTexturesCount();
+		const int32 MaterialsCount = ImportContext.FilteredScene->GetMaterialsCount();
+
+		for ( int32 TextureIndex = 0; TextureIndex < TexturesCount && !ImportContext.bUserCancelled; ++TextureIndex )
+		{
+			ImportContext.bUserCancelled |= ImportContext.Warn->ReceivedUserCancel();
+
+			TSharedPtr< IDatasmithTextureElement > TextureElement = ImportContext.FilteredScene->GetTexture( TextureIndex );
+			const FString TextureName = ObjectTools::SanitizeObjectName( TextureElement->GetName() );
+
+			for ( int32 MaterialIndex = 0; MaterialIndex < MaterialsCount; ++MaterialIndex )
+			{
+				const TSharedPtr< IDatasmithBaseMaterialElement >& BaseMaterialElement = ImportContext.FilteredScene->GetMaterial( MaterialIndex );
+
+				if ( BaseMaterialElement->IsA( EDatasmithElementType::Material ) )
+				{
+					const TSharedPtr< IDatasmithMaterialElement >& MaterialElement = StaticCastSharedPtr< IDatasmithMaterialElement >( BaseMaterialElement );
+
+					for ( int32 s = 0; s < MaterialElement->GetShadersCount(); ++s )
+					{
+						const TSharedPtr< IDatasmithShaderElement >& ShaderElement = MaterialElement->GetShader(s);
+
+						if ( FCString::Strlen( ShaderElement->GetDiffuseTexture() ) > 0 && ShaderElement->GetDiffuseTexture() == TextureName)
+						{
+							TextureElement->SetTextureMode(EDatasmithTextureMode::Diffuse);
+						}
+						else if ( FCString::Strlen( ShaderElement->GetReflectanceTexture() ) > 0 && ShaderElement->GetReflectanceTexture() == TextureName)
+						{
+							TextureElement->SetTextureMode(EDatasmithTextureMode::Specular);
+						}
+						else if ( FCString::Strlen( ShaderElement->GetDisplaceTexture() ) > 0 && ShaderElement->GetDisplaceTexture() == TextureName)
+						{
+							TextureElement->SetTextureMode(EDatasmithTextureMode::Displace);
+						}
+						else if ( FCString::Strlen( ShaderElement->GetNormalTexture() ) > 0 && ShaderElement->GetNormalTexture() == TextureName)
+						{
+							if (!ShaderElement->GetNormalTextureSampler().bInvert)
+							{
+								TextureElement->SetTextureMode(EDatasmithTextureMode::Normal);
+							}
+							else
+							{
+								TextureElement->SetTextureMode(EDatasmithTextureMode::NormalGreenInv);
+							}
+						}
+					}
+				}
+				else if ( BaseMaterialElement->IsA( EDatasmithElementType::UEPbrMaterial ) )
+				{
+					const TSharedPtr< IDatasmithUEPbrMaterialElement >& MaterialElement = StaticCastSharedPtr< IDatasmithUEPbrMaterialElement >( BaseMaterialElement );
+
+					TFunction< bool( IDatasmithMaterialExpression* ) > IsTextureConnected;
+					IsTextureConnected = [ &TextureName, &IsTextureConnected ]( IDatasmithMaterialExpression* MaterialExpression ) -> bool
+					{
+						if ( !MaterialExpression )
+						{
+							return false;
+						}
+
+						if ( MaterialExpression->IsA( EDatasmithMaterialExpressionType::Texture ) )
+						{
+							IDatasmithMaterialExpressionTexture* TextureExpression = static_cast< IDatasmithMaterialExpressionTexture* >( MaterialExpression );
+
+							if ( TextureExpression->GetTexturePathName() == TextureName )
+							{
+								return true;
+							}
+						}
+
+						for ( int32 InputIndex = 0; InputIndex < MaterialExpression->GetInputCount(); ++InputIndex )
+						{
+							IDatasmithMaterialExpression* ConnectedExpression = MaterialExpression->GetInput( InputIndex )->GetExpression();
+
+							if ( ConnectedExpression && IsTextureConnected( ConnectedExpression ) )
+							{
+								return true;
+							}
+						}
+
+						return false;
+					};
+
+					if ( IsTextureConnected( MaterialElement->GetBaseColor().GetExpression() ) )
+					{
+						TextureElement->SetTextureMode(EDatasmithTextureMode::Diffuse);
+					}
+					else if ( IsTextureConnected( MaterialElement->GetSpecular().GetExpression() ) )
+					{
+						TextureElement->SetTextureMode(EDatasmithTextureMode::Specular);
+					}
+					else if ( IsTextureConnected( MaterialElement->GetNormal().GetExpression() ) )
+					{
+						if ( TextureElement->GetTextureMode() != EDatasmithTextureMode::Bump )
+						{
+							TextureElement->SetTextureMode(EDatasmithTextureMode::Normal);
+						}
+					}
+				}
+			}
+		}
+	}
+
+	static void CompileMaterial( UObject* Material )
+	{
+		if ( !Material->IsA< UMaterialInterface >() && !Material->IsA< UMaterialFunctionInterface >() )
+		{
+			return;
+		}
+
+		TRACE_CPUPROFILER_EVENT_SCOPE(DatasmithImporterImpl::CompileMaterial);
+
+		FMaterialUpdateContext MaterialUpdateContext;
+
+		if ( UMaterialInterface* MaterialInterface = Cast< UMaterialInterface >( Material ) )
+		{
+			MaterialUpdateContext.AddMaterialInterface( MaterialInterface );
+		}
+
+		if ( UMaterialInstanceConstant* ConstantMaterialInstance = Cast< UMaterialInstanceConstant >( Material ) )
+		{
+			// If BlendMode override property has been changed, make sure this combination of the parent material is compiled
+			if ( ConstantMaterialInstance->BasePropertyOverrides.bOverride_BlendMode == true )
+			{
+				ConstantMaterialInstance->ForceRecompileForRendering();
+			}
+			else
+			{
+				// If a static switch is overriden, we need to recompile
+				FStaticParameterSet StaticParameters;
+				ConstantMaterialInstance->GetStaticParameterValues( StaticParameters );
+
+				for ( FStaticSwitchParameter& Switch : StaticParameters.StaticSwitchParameters )
+				{
+					if ( Switch.bOverride )
+					{
+						ConstantMaterialInstance->ForceRecompileForRendering();
+						break;
+					}
+				}
+			}
+		}
+
+		Material->PreEditChange( nullptr );
+		Material->PostEditChange();
+	}
+
+	static void FixReferencesForObject( UObject* Object, const TMap< UObject*, UObject* >& ReferencesToRemap )
+	{
+		constexpr bool bNullPrivateRefs = false;
+		constexpr bool bIgnoreOuterRef = true;
+		constexpr bool bIgnoreArchetypeRef = true;
+
+		if ( ReferencesToRemap.Num() > 0 )
+		{
+			FArchiveReplaceObjectRef< UObject > ArchiveReplaceObjectRef( Object, ReferencesToRemap, bNullPrivateRefs, bIgnoreOuterRef, bIgnoreArchetypeRef );
+		}
+	}
+
+	using FMigratedTemplatePairType = TPair< TStrongObjectPtr< UDatasmithObjectTemplate >, TStrongObjectPtr< UDatasmithObjectTemplate > >;
+
+	/**
+	 * Creates templates to apply the values from the SourceObject on the DestinationObject.
+	 *
+	 * @returns An array of template pairs. The key is the template for the object, the value is a template to force apply to the object,
+	 *			it contains the values from the key and any overrides that were present on the DestinationObject.
+	 */
+	static TArray< FMigratedTemplatePairType > MigrateTemplates( UObject* SourceObject, UObject* DestinationObject, const TMap< UObject*, UObject* >* ReferencesToRemap, bool bIsForActor )
+	{
+		TArray< FMigratedTemplatePairType > Results;
+
+		if ( !SourceObject )
+		{
+			return Results;
+		}
+
+		TMap< TSubclassOf< UDatasmithObjectTemplate >, UDatasmithObjectTemplate* >* SourceTemplates = FDatasmithObjectTemplateUtils::FindOrCreateObjectTemplates( SourceObject );
+
+		if ( !SourceTemplates )
+		{
+			return Results;
+		}
+
+		for ( const TPair< TSubclassOf< UDatasmithObjectTemplate >, UDatasmithObjectTemplate* >& SourceTemplatePair : *SourceTemplates )
+		{
+			if ( bIsForActor == SourceTemplatePair.Value->bIsActorTemplate )
+			{
+				FMigratedTemplatePairType& Result = Results.AddDefaulted_GetRef();
+
+				TStrongObjectPtr< UDatasmithObjectTemplate > SourceTemplate{ NewObject< UDatasmithObjectTemplate >(GetTransientPackage(), SourceTemplatePair.Key.Get()) }; // The SourceTemplate is the one we will persist so set its outer as DestinationObject
+
+				SourceTemplate->Load(SourceObject);
+
+				if ( ReferencesToRemap )
+				{
+					FixReferencesForObject(SourceTemplate.Get(), *ReferencesToRemap);
+				}
+
+				Result.Key = SourceTemplate;
+
+				if (DestinationObject && !DestinationObject->IsPendingKillOrUnreachable())
+				{
+					Result.Value = TStrongObjectPtr< UDatasmithObjectTemplate >(UDatasmithObjectTemplate::GetDifference( DestinationObject, SourceTemplate.Get()));
+				}
+				else
+				{
+					Result.Value = SourceTemplate;
+				}
+			}
+		}
+
+		return Results;
+	}
+
+	/**
+	 * Applies the templates created from MigrateTemplates to DestinationObject.
+	 *
+	 * For an Object A that should be duplicated over an existing A', for which we want to keep the Datasmith overrides:
+	 * - Call MigrateTemplates(A, A')
+	 * - Duplicate A over A'
+	 * - ApplyMigratedTemplates(A')
+	 */
+	static void ApplyMigratedTemplates( TArray< FMigratedTemplatePairType >& MigratedTemplates, UObject* DestinationObject )
+	{
+		for ( FMigratedTemplatePairType& MigratedTemplate : MigratedTemplates )
+		{
+			UDatasmithObjectTemplate* SourceTemplate = MigratedTemplate.Key.Get();
+			UDatasmithObjectTemplate* DestinationTemplate = MigratedTemplate.Value.Get();
+
+			DestinationTemplate->Apply(DestinationObject, true); // Restore the overrides
+			FDatasmithObjectTemplateUtils::SetObjectTemplate(DestinationObject, SourceTemplate); // Set SourceTemplate as our template so that any differences are considered overrides
+
+		}
+	}
+
+	static UObject* FinalizeAsset( UObject* SourceAsset, const TCHAR* AssetPath, UObject* ExistingAsset, TMap< UObject*, UObject* >* ReferencesToRemap )
+	{
+		if ( ReferencesToRemap )
+		{
+			FixReferencesForObject( SourceAsset, *ReferencesToRemap );
+		}
+
+		TArray< FMigratedTemplatePairType > MigratedTemplates = MigrateTemplates( SourceAsset, ExistingAsset, ReferencesToRemap, false );
+
+		UObject* FinalAsset = PublicizeAsset( SourceAsset, AssetPath, ExistingAsset );
+
+		ApplyMigratedTemplates( MigratedTemplates, FinalAsset );
+
+		if ( ReferencesToRemap && SourceAsset && SourceAsset != FinalAsset )
+		{
+			ReferencesToRemap->Add( SourceAsset, FinalAsset );
+		}
+
+		return FinalAsset;
+	}
+
+	class FActorWriter final: public FObjectWriter
+	{
+	public:
+		FActorWriter( UObject* Object, TArray< uint8 >& Bytes )
+			: FObjectWriter( Bytes )
+		{
+			SetIsLoading(false);
+			SetIsSaving(true);
+			SetIsPersistent(false);
+
+			Object->Serialize(*this); // virtual call in ctr -> final class
+		}
+
+		virtual bool ShouldSkipProperty(const UProperty* InProperty) const override
+		{
+			bool bSkip = false;
+
+			if ( InProperty->IsA< UObjectPropertyBase >() )
+			{
+				bSkip = true;
+			}
+			else if ( InProperty->HasAnyPropertyFlags(CPF_Transient) || !InProperty->HasAnyPropertyFlags(CPF_Edit | CPF_Interp) )
+			{
+				bSkip = true;
+			}
+
+			return bSkip;
+		}
+	};
+
+	class FComponentWriter final : public FObjectWriter
+	{
+	public:
+		FComponentWriter( UObject* Object, TArray< uint8 >& Bytes )
+			: FObjectWriter( Bytes )
+		{
+			SetIsLoading(false);
+			SetIsSaving(true);
+			SetIsPersistent(false);
+
+			Object->Serialize(*this); // virtual call in ctr -> final class
+		}
+
+		virtual bool ShouldSkipProperty(const UProperty* InProperty) const override
+		{
+			bool bSkip = false;
+
+			if ( InProperty->HasAnyPropertyFlags(CPF_Transient) || !InProperty->HasAnyPropertyFlags(CPF_Edit | CPF_Interp) )
+			{
+				bSkip = true;
+			}
+
+			return bSkip;
+		}
+	};
+
+	static void DeleteImportSceneActorIfNeeded(FDatasmithActorImportContext& ActorContext, bool bForce = false)
+	{
+		ADatasmithSceneActor*& ImportSceneActor = ActorContext.ImportSceneActor;
+		if ( !ActorContext.FinalSceneActors.Contains(ImportSceneActor) || bForce )
+		{
+			if ( ImportSceneActor )
+			{
+				for ( const auto& Pair : ImportSceneActor->RelatedActors )
+				{
+					AActor* RelatedActor = Pair.Value.Get();
+					if ( RelatedActor )
+					{
+						FDatasmithImporterUtils::DeleteActor( *RelatedActor );
+					}
+				}
+
+				FDatasmithImporterUtils::DeleteActor( *ImportSceneActor );
+
+				// Null also the ImportSceneActor from the Actor Context because it's a ref to it.
+				ImportSceneActor = nullptr;
+			}
+		}
+	}
+
+	static UActorComponent* PublicizeComponent(UActorComponent& SourceComponent, UActorComponent* DestinationComponent, AActor& DestinationActor, TMap< UObject*, UObject* >& ReferencesToRemap, USceneComponent* DestinationParent = nullptr )
+	{
+		TRACE_CPUPROFILER_EVENT_SCOPE(DatasmithImporterImpl::PublicizeComponent);
+
+		SafetyGarbageCollect();
+
+		if ( !SourceComponent.HasAnyFlags( RF_Transient | RF_TextExportTransient ) )
+		{
+			if (!DestinationComponent || DestinationComponent->IsPendingKillOrUnreachable())
+			{
+				if (DestinationComponent)
+				{
+					// Change the name of the old component so that the new object won't recycle the old one.
+					DestinationComponent->Rename(nullptr, nullptr, REN_DontCreateRedirectors | REN_NonTransactional);
+				}
+
+				if ( UActorComponent* OldComponent = static_cast<UActorComponent*>( FindObjectWithOuter( &DestinationActor, UActorComponent::StaticClass(), SourceComponent.GetFName() ) ) )
+				{
+					OldComponent->DestroyComponent( true );
+					// Change the name of the old component so that the new object won't recycle the old one.
+					OldComponent->Rename(nullptr, nullptr, REN_DontCreateRedirectors | REN_NonTransactional);
+				}
+
+				// Create a new component
+				DestinationComponent = NewObject<UActorComponent>(&DestinationActor, SourceComponent.GetClass(), SourceComponent.GetFName(), RF_Transactional);
+				DestinationActor.AddInstanceComponent(DestinationComponent);
+
+				if (USceneComponent* NewSceneComponent = Cast<USceneComponent>(DestinationComponent))
+				{
+					if ( !DestinationActor.GetRootComponent() )
+					{
+						DestinationActor.SetRootComponent(NewSceneComponent);
+					}
+					if ( DestinationParent )
+					{
+						NewSceneComponent->AttachToComponent( DestinationParent, FAttachmentTransformRules::KeepRelativeTransform );
+					}
+				}
+
+				DestinationComponent->RegisterComponent();
+			}
+
+			check(DestinationComponent);
+
+			// Copy component data
+			{
+				TArray< uint8 > Bytes;
+				FComponentWriter ObjectWriter(&SourceComponent, Bytes);
+				FObjectReader ObjectReader(DestinationComponent, Bytes);
+			}
+
+			FixReferencesForObject(DestinationComponent, ReferencesToRemap);
+
+			// #ueent_todo: we shouldn't be copying instanced object pointers in the first place
+			UDatasmithAssetUserData* SourceAssetUserData = DestinationComponent->GetAssetUserData< UDatasmithAssetUserData >();
+
+			if ( SourceAssetUserData )
+			{
+				UDatasmithAssetUserData* DestinationAssetUserData = DuplicateObject< UDatasmithAssetUserData >( SourceAssetUserData, DestinationComponent );
+				DestinationComponent->RemoveUserDataOfClass( UDatasmithAssetUserData::StaticClass() );
+				DestinationComponent->AddAssetUserData( DestinationAssetUserData );
+			}
+
+			ReferencesToRemap.Add( &SourceComponent, DestinationComponent );
+
+			return DestinationComponent;
+		}
+
+		return nullptr;
+	}
+
+	static void FinalizeSceneComponent(FDatasmithImportContext& ImportContext, USceneComponent& SourceComponent, AActor& DestinationActor, USceneComponent* DestinationParent, TMap<UObject *, UObject *>& ReferencesToRemap )
+	{
+		USceneComponent* DestinationComponent = static_cast<USceneComponent*>( FindObjectWithOuter( &DestinationActor, SourceComponent.GetClass(), SourceComponent.GetFName() ) );
+		FName SourceComponentDatasmithId = FDatasmithImporterUtils::GetDatasmithElementId(&SourceComponent);
+
+		if ( SourceComponentDatasmithId.IsNone() )
+		{
+			// This component is not tracked by datasmith
+			if ( !DestinationComponent || DestinationComponent->IsPendingKillOrUnreachable() )
+			{
+				DestinationComponent = static_cast<USceneComponent*> ( PublicizeComponent(SourceComponent, DestinationComponent, DestinationActor, ReferencesToRemap, DestinationParent) );
+				if ( DestinationComponent )
+				{
+					// Put back the components in a proper state
+					DestinationComponent->UpdateComponentToWorld();
+				}
+			}
+		}
+		else
+		{
+			check(ImportContext.ActorsContext.CurrentTargetedScene);
+
+			TArray< FMigratedTemplatePairType > MigratedTemplates = MigrateTemplates(
+				&SourceComponent,
+				DestinationComponent,
+				&ReferencesToRemap,
+				false
+			);
+
+			DestinationComponent = static_cast< USceneComponent* > ( PublicizeComponent( SourceComponent, DestinationComponent, DestinationActor, ReferencesToRemap, DestinationParent ) );
+
+			if ( DestinationComponent )
+			{
+				// Put back the components in a proper state (Unfortunately without this the set relative transform might not work)
+				DestinationComponent->UpdateComponentToWorld();
+				ApplyMigratedTemplates( MigratedTemplates, DestinationComponent );
+				DestinationComponent->PostEditChange();
+			}
+		}
+
+		USceneComponent* AttachParentForChildren = DestinationComponent  ? DestinationComponent : DestinationParent;
+		for ( USceneComponent* Child : SourceComponent.GetAttachChildren() )
+		{
+			// Only finalize components that are from the same actor
+			if ( Child && Child->GetOuter() == SourceComponent.GetOuter() )
+			{
+				FinalizeSceneComponent( ImportContext, *Child, DestinationActor, AttachParentForChildren, ReferencesToRemap );
+			}
+		}
+	}
+
+	static void FinalizeComponents(FDatasmithImportContext& ImportContext, AActor& SourceActor, AActor& DestinationActor, TMap<UObject *, UObject *>& ReferencesToRemap)
+	{
+		USceneComponent* ParentComponent = nullptr;
+
+		// Find the parent component
+		UObject** ObjectPtr = ReferencesToRemap.Find( SourceActor.GetRootComponent()->GetAttachParent() );
+		if ( ObjectPtr )
+		{
+			ParentComponent = Cast<USceneComponent>( *ObjectPtr );
+		}
+
+		// Finalize the scene components recursively
+		{
+			USceneComponent* RootComponent = SourceActor.GetRootComponent();
+			if ( RootComponent )
+			{
+				FinalizeSceneComponent( ImportContext, *RootComponent, DestinationActor, ParentComponent, ReferencesToRemap );
+			}
+		}
+
+
+		for ( UActorComponent* SourceComponent : SourceActor.GetComponents() )
+		{
+			// Only the non scene component haven't been finalize
+			if ( SourceComponent && !SourceComponent->GetClass()->IsChildOf<USceneComponent>() )
+			{
+				UActorComponent* DestinationComponent = static_cast< UActorComponent* >( FindObjectWithOuter( &DestinationActor, SourceComponent->GetClass(), SourceComponent->GetFName() ) );
+				if (!DestinationComponent)
+				{
+					DestinationComponent = PublicizeComponent( *SourceComponent, DestinationComponent, DestinationActor, ReferencesToRemap );
+				}
+			}
+		}
+	}
+
+	static void GatherUnsupportedVirtualTexturesAndMaterials(const TMap<TSharedRef< IDatasmithBaseMaterialElement >, UMaterialInterface*>& ImportedMaterials, TSet<UTexture2D*>& VirtualTexturesToConvert, TArray<UMaterial*>& MaterialsToRefreshAfterVirtualTextureConversion)
+	{
+		//Multimap cache to avoid parsing the same base material multiple times.
+		TMultiMap<UMaterial*, FMaterialParameterInfo> TextureParametersToConvertMap;
+
+		//Loops through all imported material instances and add to VirtualTexturesToConvert all the virtual texture parameters that don't support virtual texturing in the base material.
+		for (const TPair< TSharedRef< IDatasmithBaseMaterialElement >, UMaterialInterface* >& ImportedMaterialPair : ImportedMaterials)
+		{
+			UMaterialInterface* CurrentMaterialInterface = ImportedMaterialPair.Value;
+			UMaterial* BaseMaterial = CurrentMaterialInterface->GetMaterial();
+
+			if(!TextureParametersToConvertMap.Contains(BaseMaterial))
+			{
+				bool bRequiresTextureCheck = false;
+				TArray<FMaterialParameterInfo> OutParameterInfo;
+				TArray<FGuid> Guids;
+				BaseMaterial->GetAllTextureParameterInfo(OutParameterInfo, Guids);
+
+				for (int32 ParameterInfoIndex = 0; ParameterInfoIndex < OutParameterInfo.Num(); ++ParameterInfoIndex)
+				{
+					UTexture* TextureParameter = nullptr;
+
+					if (BaseMaterial->GetTextureParameterValue(OutParameterInfo[ParameterInfoIndex], TextureParameter) && VirtualTexturesToConvert.Contains(Cast<UTexture2D>(TextureParameter)))
+					{
+						bRequiresTextureCheck = true;
+						TextureParametersToConvertMap.Add(BaseMaterial, OutParameterInfo[ParameterInfoIndex]);
+					}
+				}
+
+				if (bRequiresTextureCheck)
+				{
+					MaterialsToRefreshAfterVirtualTextureConversion.Add(BaseMaterial);
+				}
+				else
+				{
+					//Adding a dummy MaterialParameterInfo so that we don't have to parse this Base Material again.
+					TextureParametersToConvertMap.Add(BaseMaterial, FMaterialParameterInfo());
+
+					//If no unsupported virtual texture parameters were found, it's possible that a texture needing conversion is simply not exposed as a parameter, so we still need to check for those.
+					TArray<UObject*> ReferencedTextures;
+					BaseMaterial->AppendReferencedTextures(ReferencedTextures);
+					for (UObject* ReferencedTexture : ReferencedTextures)
+					{
+						if (VirtualTexturesToConvert.Contains(Cast<UTexture2D>(ReferencedTexture)))
+						{
+							MaterialsToRefreshAfterVirtualTextureConversion.Add(BaseMaterial);
+							break;
+						}
+					}
+				}
+			}
+
+			for (auto ParameterInfoIt = TextureParametersToConvertMap.CreateKeyIterator(BaseMaterial); ParameterInfoIt; ++ParameterInfoIt)
+			{
+				UTexture* TextureParameter = nullptr;
+
+				if (CurrentMaterialInterface->GetTextureParameterValue(ParameterInfoIt.Value(), TextureParameter) &&
+					TextureParameter && TextureParameter->VirtualTextureStreaming)
+				{
+					if (UTexture2D* TextureToConvert = Cast<UTexture2D>(TextureParameter))
+					{
+						VirtualTexturesToConvert.Add(TextureToConvert);
+					}
+				}
+			}
+		}
+	}
+
+	static void ConvertUnsupportedVirtualTexture(FDatasmithImportContext& ImportContext, TSet<UTexture2D*>& VirtualTexturesToConvert, const TMap<UObject*, UObject*>& ReferencesToRemap)
+	{
+		TArray<UMaterial*> MaterialsToRefreshAfterVirtualTextureConversion;
+		GatherUnsupportedVirtualTexturesAndMaterials(ImportContext.ImportedMaterials, ImportContext.AssetsContext.VirtualTexturesToConvert, MaterialsToRefreshAfterVirtualTextureConversion);
+
+		if (VirtualTexturesToConvert.Num() != 0)
+		{
+			for (UTexture2D*& TextureToConvert : VirtualTexturesToConvert)
+			{
+				if (UObject* const* RemappedTexture = ReferencesToRemap.Find(TextureToConvert))
+				{
+					TextureToConvert = Cast<UTexture2D>(*RemappedTexture);
+				}
+
+				ImportContext.LogWarning(FText::Format(LOCTEXT("DatasmithVirtualTextureConverted", "The imported texture {0} could not be imported as virtual texture as it is not supported in all the materials using it."), FText::FromString(TextureToConvert->GetName())));
+			}
+
+			for (int32 MaterialIndex = 0; MaterialIndex < MaterialsToRefreshAfterVirtualTextureConversion.Num(); ++MaterialIndex)
+			{
+				if (UObject* const* RemappedMaterial = ReferencesToRemap.Find(MaterialsToRefreshAfterVirtualTextureConversion[MaterialIndex]))
+				{
+					MaterialsToRefreshAfterVirtualTextureConversion[MaterialIndex] = Cast<UMaterial>(*RemappedMaterial);
+				}
+			}
+
+			TArray<UTexture2D*> TexturesToConvertList(VirtualTexturesToConvert.Array());
+			IAssetTools& AssetTools = FModuleManager::LoadModuleChecked<FAssetToolsModule>("AssetTools").Get();
+			AssetTools.ConvertVirtualTextures(TexturesToConvertList, true, &MaterialsToRefreshAfterVirtualTextureConversion);
+		}
+	}
+}
+
+void FDatasmithImporter::ImportStaticMeshes( FDatasmithImportContext& ImportContext )
+{
+	const int32 StaticMeshesCount = ImportContext.FilteredScene->GetMeshesCount();
+
+	if ( !ImportContext.Options->BaseOptions.bIncludeGeometry || StaticMeshesCount == 0 )
+	{
+		return;
+	}
+
+	TRACE_CPUPROFILER_EVENT_SCOPE(FDatasmithImporter::ImportStaticMeshes);
+
+	FScopedSlowTask Progress(StaticMeshesCount, LOCTEXT("ImportStaticMeshes", "Importing Static Meshes..."), true, *ImportContext.Warn );
+	Progress.MakeDialog(true);
+
+	TMap<TSharedRef<IDatasmithMeshElement>, TFuture<FDatasmithMeshElementPayload*>> MeshElementPayloads;
+
+	FDatasmithTranslatorCapabilities TranslatorCapabilities;
+	if (ImportContext.SceneTranslator)
+	{
+		ImportContext.SceneTranslator->Initialize(TranslatorCapabilities);
+	}
+
+	// Parallelize loading by doing a first pass to send translator loading into async task
+	if (TranslatorCapabilities.bParallelLoadStaticMeshSupported)
+	{
+		for (int32 MeshIndex = 0; MeshIndex < StaticMeshesCount && !ImportContext.bUserCancelled; ++MeshIndex)
+		{
+			ImportContext.bUserCancelled |= ImportContext.Warn->ReceivedUserCancel();
+
+			if (!ImportContext.AssetsContext.StaticMeshesFinalPackage || ImportContext.AssetsContext.StaticMeshesFinalPackage->GetFName() == NAME_None || ImportContext.SceneTranslator == nullptr)
+			{
+				continue;
+			}
+
+			TSharedRef<IDatasmithMeshElement> MeshElement = ImportContext.FilteredScene->GetMesh( MeshIndex ).ToSharedRef();
+
+			UStaticMesh*& ImportedStaticMesh = ImportContext.ImportedStaticMeshes.FindOrAdd( MeshElement );
+
+			// We still have factories that are importing the UStaticMesh on their own, so check if it's already imported here
+			if (ImportedStaticMesh == nullptr)
+			{
+				// Parallel loading from the translator using futures
+				MeshElementPayloads.Add(
+					MeshElement,
+					Async(
+						EAsyncExecution::LargeThreadPool,
+						[&ImportContext, MeshElement]() -> FDatasmithMeshElementPayload*
+						{
+							if (ImportContext.bUserCancelled)
+							{
+								return nullptr;
+							}
+
+							TUniquePtr<FDatasmithMeshElementPayload> MeshPayload = MakeUnique<FDatasmithMeshElementPayload>();
+							return ImportContext.SceneTranslator->LoadStaticMesh(MeshElement, *MeshPayload) ? MeshPayload.Release() : nullptr;
+						}
+					)
+				);
+			}
+		}
+	}
+
+	// This pass will wait on the futures we got from the first pass async tasks
+	for ( int32 MeshIndex = 0; MeshIndex < StaticMeshesCount && !ImportContext.bUserCancelled; ++MeshIndex )
+	{
+		ImportContext.bUserCancelled |= ImportContext.Warn->ReceivedUserCancel();
+
+		TSharedRef< IDatasmithMeshElement > MeshElement = ImportContext.FilteredScene->GetMesh( MeshIndex ).ToSharedRef();
+
+		Progress.EnterProgressFrame( 1.f, FText::FromString( FString::Printf( TEXT("Importing static mesh %d/%d (%s) ..."), MeshIndex + 1, StaticMeshesCount, MeshElement->GetLabel() ) ) );
+
+		UStaticMesh* ExistingStaticMesh = nullptr;
+
+		if (ImportContext.SceneAsset)
+		{
+			TSoftObjectPtr< UStaticMesh >* ExistingStaticMeshPtr = ImportContext.SceneAsset->StaticMeshes.Find( MeshElement->GetName() );
+
+			if (ExistingStaticMeshPtr)
+			{
+				ExistingStaticMesh = ExistingStaticMeshPtr->LoadSynchronous();
+			}
+		}
+
+		// #ueent_todo rewrite in N passes:
+		//  - GetDestination (find or create StaticMesh, duplicate, flags and context etc)
+		//  - Import (Import data in simple memory repr (eg. TArray<FMeshDescription>)
+		//  - Set (fill UStaticMesh with imported data)
+		TFuture<FDatasmithMeshElementPayload*> MeshPayload;
+		if (MeshElementPayloads.RemoveAndCopyValue(MeshElement, MeshPayload))
+		{
+			TUniquePtr<FDatasmithMeshElementPayload> MeshPayloadPtr(MeshPayload.Get());
+			if (MeshPayloadPtr.IsValid())
+			{
+				ImportStaticMesh(ImportContext, MeshElement, ExistingStaticMesh, MeshPayloadPtr.Get());
+			}
+		}
+		else
+		{
+			ImportStaticMesh(ImportContext, MeshElement, ExistingStaticMesh, nullptr);
+		}
+	}
+
+	//Just make sure there is no async task left running in case of a cancellation
+	for ( const TPair<TSharedRef<IDatasmithMeshElement>, TFuture<FDatasmithMeshElementPayload*>> & Kvp : MeshElementPayloads)
+	{
+		// Wait for the result and delete it when getting out of scope
+		TUniquePtr<FDatasmithMeshElementPayload> MeshPayloadPtr(Kvp.Value.Get());
+	}
+
+	TMap< TSharedRef< IDatasmithMeshElement >, float > LightmapWeights = FDatasmithStaticMeshImporter::CalculateMeshesLightmapWeights( ImportContext.Scene.ToSharedRef() );
+
+	for ( TPair< TSharedRef< IDatasmithMeshElement >, UStaticMesh* >& ImportedStaticMeshPair : ImportContext.ImportedStaticMeshes )
+	{
+		FDatasmithStaticMeshImporter::SetupStaticMesh( ImportContext.AssetsContext, ImportedStaticMeshPair.Key, ImportedStaticMeshPair.Value, ImportContext.Options->BaseOptions.StaticMeshOptions, LightmapWeights[ ImportedStaticMeshPair.Key ] );
+	}
+}
+
+UStaticMesh* FDatasmithImporter::ImportStaticMesh( FDatasmithImportContext& ImportContext, TSharedRef< IDatasmithMeshElement > MeshElement, UStaticMesh* ExistingStaticMesh, FDatasmithMeshElementPayload* MeshPayload)
+{
+	if ( !ImportContext.AssetsContext.StaticMeshesFinalPackage || ImportContext.AssetsContext.StaticMeshesFinalPackage->GetFName() == NAME_None || ImportContext.SceneTranslator == nullptr)
+	{
+		return nullptr;
+	}
+
+	TRACE_CPUPROFILER_EVENT_SCOPE(FDatasmithImporter::ImportStaticMesh);
+
+	UStaticMesh*& ImportedStaticMesh = ImportContext.ImportedStaticMeshes.FindOrAdd( MeshElement );
+
+	TArray<UDatasmithAdditionalData*> AdditionalData;
+
+	if ( ImportedStaticMesh == nullptr ) // We still have factories that are importing the UStaticMesh on their own, so check if it's already imported here
+	{
+		FDatasmithMeshElementPayload LocalMeshPayload;
+		if (MeshPayload == nullptr)
+		{
+			ImportContext.SceneTranslator->LoadStaticMesh(MeshElement, LocalMeshPayload);
+			MeshPayload = &LocalMeshPayload;
+		}
+
+		ImportedStaticMesh = FDatasmithStaticMeshImporter::ImportStaticMesh( MeshElement, *MeshPayload, ImportContext.ObjectFlags & ~RF_Public, ImportContext.Options->BaseOptions.StaticMeshOptions, ImportContext.AssetsContext, ExistingStaticMesh );
+		AdditionalData = MoveTemp(MeshPayload->AdditionalData);
+
+		// Creation of static mesh failed, remove it from the list of importer mesh elements
+		if (ImportedStaticMesh == nullptr)
+		{
+			ImportContext.ImportedStaticMeshes.Remove(MeshElement);
+			return nullptr;
+		}
+	}
+
+	CreateStaticMeshAssetImportData( ImportContext, MeshElement, ImportedStaticMesh, AdditionalData );
+
+	ImportMetaDataForObject( ImportContext, MeshElement, ImportedStaticMesh );
+
+	if ( MeshElement->GetLightmapSourceUV() >= MAX_MESH_TEXTURE_COORDS_MD )
+	{
+		FFormatNamedArguments FormatArgs;
+		FormatArgs.Add(TEXT("SourceUV"), FText::FromString(FString::FromInt(MeshElement->GetLightmapSourceUV())));
+		FormatArgs.Add(TEXT("MeshName"), FText::FromName(ImportedStaticMesh->GetFName()));
+		ImportContext.LogError(FText::Format(LOCTEXT("InvalidLightmapSourceUVError", "The lightmap source UV '{SourceUV}' used for the lightmap UV generation of the mesh '{MeshName}' is invalid."), FormatArgs));
+	}
+
+	return ImportedStaticMesh;
+}
+
+
+UStaticMesh* FDatasmithImporter::FinalizeStaticMesh( UStaticMesh* SourceStaticMesh, const TCHAR* StaticMeshesFolderPath, UStaticMesh* ExistingStaticMesh, TMap< UObject*, UObject* >* ReferencesToRemap, bool bBuild)
+{
+	using namespace DatasmithImporterImpl;
+
+	UStaticMesh* DestinationStaticMesh = Cast< UStaticMesh >( FinalizeAsset( SourceStaticMesh, StaticMeshesFolderPath, ExistingStaticMesh, ReferencesToRemap ) );
+
+	if (bBuild)
+	{
+		FDatasmithStaticMeshImporter::BuildStaticMesh(DestinationStaticMesh);
+	}
+
+	return DestinationStaticMesh;
+}
+
+void FDatasmithImporter::CreateStaticMeshAssetImportData(FDatasmithImportContext& InContext, TSharedRef< IDatasmithMeshElement > MeshElement, UStaticMesh* ImportedStaticMesh, TArray<UDatasmithAdditionalData*>& AdditionalData)
+{
+	UDatasmithStaticMeshImportData::DefaultOptionsPair ImportOptions = UDatasmithStaticMeshImportData::DefaultOptionsPair( InContext.Options->BaseOptions.StaticMeshOptions, InContext.Options->BaseOptions.AssetOptions );
+
+	UDatasmithStaticMeshImportData* MeshImportData = UDatasmithStaticMeshImportData::GetImportDataForStaticMesh( ImportedStaticMesh, ImportOptions );
+
+	if ( MeshImportData )
+	{
+		// Update the import data source file and set the mesh hash
+		// #ueent_todo FH: piggybacking off of the SourceData file hash for now, until we have custom derived AssetImportData properly serialize to the AssetRegistry
+		FMD5Hash Hash = MeshElement->CalculateElementHash( false );
+		MeshImportData->Update( InContext.Options->FilePath, &Hash );
+
+		// Set the final outer // #ueent_review: propagate flags of outer?
+		for (UDatasmithAdditionalData* Data: AdditionalData)
+		{
+			Data->Rename(nullptr, MeshImportData);
+		}
+		MeshImportData->AdditionalData = MoveTemp(AdditionalData);
+	}
+}
+
+void FDatasmithImporter::ImportTextures( FDatasmithImportContext& ImportContext )
+{
+	TRACE_CPUPROFILER_EVENT_SCOPE(FDatasmithImporter::ImportTextures);
+
+	DatasmithImporterImpl::SetTexturesMode( ImportContext );
+
+	const int32 TexturesCount = ImportContext.FilteredScene->GetTexturesCount();
+
+	FScopedSlowTask Progress( (float)TexturesCount, LOCTEXT("ImportSlowTask", "Importing Textures..."), true, *ImportContext.Warn );
+	Progress.MakeDialog(true);
+
+	if (ImportContext.Options->TextureConflictPolicy != EDatasmithImportAssetConflictPolicy::Ignore && TexturesCount > 0)
+	{
+		FDatasmithTextureImporter DatasmithTextureImporter(ImportContext);
+
+		for ( int32 i = 0; i < TexturesCount && !ImportContext.bUserCancelled; i++ )
+		{
+			ImportContext.bUserCancelled |= ImportContext.Warn->ReceivedUserCancel();
+
+			TSharedPtr< IDatasmithTextureElement > TextureElement = ImportContext.FilteredScene->GetTexture(i);
+
+			if ( !TextureElement )
+			{
+				continue;
+			}
+
+			Progress.EnterProgressFrame( 1.f, FText::FromString( FString::Printf( TEXT("Importing texture %d/%d (%s) ..."), i + 1, TexturesCount, TextureElement->GetLabel() ) ) );
+
+			UTexture* ExistingTexture = nullptr;
+
+			if ( ImportContext.SceneAsset )
+			{
+				TSoftObjectPtr< UTexture >* ExistingTexturePtr = ImportContext.SceneAsset->Textures.Find( TextureElement->GetName() );
+
+				if ( ExistingTexturePtr )
+				{
+					ExistingTexture = ExistingTexturePtr->LoadSynchronous();
+				}
+			}
+
+			ImportTexture( ImportContext, DatasmithTextureImporter, TextureElement.ToSharedRef(), ExistingTexture );
+
+		}
+	}
+}
+
+UTexture* FDatasmithImporter::ImportTexture( FDatasmithImportContext& ImportContext, FDatasmithTextureImporter& DatasmithTextureImporter, TSharedRef< IDatasmithTextureElement > TextureElement, UTexture* ExistingTexture )
+{
+	TRACE_CPUPROFILER_EVENT_SCOPE(FDatasmithImporter::ImportTexture);
+
+	UTexture*& ImportedTexture = ImportContext.ImportedTextures.FindOrAdd( TextureElement );
+	ImportedTexture = DatasmithTextureImporter.CreateTexture( TextureElement );
+
+	if (ImportedTexture == nullptr)
+	{
+		ImportContext.ImportedTextures.Remove( TextureElement );
+		return nullptr;
+	}
+
+	ImportMetaDataForObject( ImportContext, TextureElement, ImportedTexture );
+
+	return ImportedTexture;
+}
+
+UTexture* FDatasmithImporter::FinalizeTexture( UTexture* SourceTexture, const TCHAR* TexturesFolderPath, UTexture* ExistingTexture, TMap< UObject*, UObject* >* ReferencesToRemap )
+{
+	return Cast< UTexture >( DatasmithImporterImpl::FinalizeAsset( SourceTexture, TexturesFolderPath, ExistingTexture, ReferencesToRemap ) );
+}
+
+void FDatasmithImporter::ImportMaterials( FDatasmithImportContext& ImportContext )
+{
+	TRACE_CPUPROFILER_EVENT_SCOPE(FDatasmithImporter::ImportMaterials);
+
+	if ( ImportContext.Options->MaterialConflictPolicy != EDatasmithImportAssetConflictPolicy::Ignore && ImportContext.FilteredScene->GetMaterialsCount() > 0 )
+	{
+		IDatasmithShaderElement::bUseRealisticFresnel = ( ImportContext.Options->MaterialQuality == EDatasmithImportMaterialQuality::UseRealFresnelCurves );
+		IDatasmithShaderElement::bDisableReflectionFresnel = ( ImportContext.Options->MaterialQuality == EDatasmithImportMaterialQuality::UseNoFresnelCurves );
+
+		//Import referenced materials as MaterialFunctions first
+		for ( TSharedPtr< IDatasmithBaseMaterialElement > MaterialElement : FDatasmithImporterUtils::GetOrderedListOfMaterialsReferencedByMaterials( ImportContext.FilteredScene ) )
+		{
+			ImportMaterialFunction(ImportContext, MaterialElement.ToSharedRef() );
+		}
+
+		ImportContext.AssetsContext.MaterialsRequirements.Empty( ImportContext.FilteredScene->GetMaterialsCount() );
+
+		for ( int32 MaterialIndex = 0; MaterialIndex < ImportContext.FilteredScene->GetMaterialsCount(); ++MaterialIndex )
+		{
+			TSharedRef< IDatasmithBaseMaterialElement > MaterialElement = ImportContext.FilteredScene->GetMaterial( MaterialIndex ).ToSharedRef();
+
+			UMaterialInterface* ExistingMaterial = nullptr;
+
+			if ( ImportContext.SceneAsset )
+			{
+				TSoftObjectPtr< UMaterialInterface >* ExistingMaterialPtr = ImportContext.SceneAsset->Materials.Find( MaterialElement->GetName() );
+
+				if ( ExistingMaterialPtr )
+				{
+					ExistingMaterial = ExistingMaterialPtr->LoadSynchronous();
+				}
+			}
+
+			ImportMaterial( ImportContext, MaterialElement, ExistingMaterial );
+		}
+
+		// IMPORTANT: FGlobalComponentReregisterContext destructor will de-register and re-register all UActorComponent present in the world
+		// Consequently, all static meshes will stop using the FMaterialResource of the original materials on de-registration
+		// and will use the new FMaterialResource created on re-registration.
+		// Otherwise, the editor will crash on redraw
+		FGlobalComponentReregisterContext RecreateComponents;
+	}
+}
+
+UMaterialFunction* FDatasmithImporter::ImportMaterialFunction(FDatasmithImportContext& ImportContext, TSharedRef< IDatasmithBaseMaterialElement > MaterialElement)
+{
+	UMaterialFunction* ImportedMaterialFunction = FDatasmithMaterialImporter::CreateMaterialFunction( ImportContext, MaterialElement );
+
+	if (!ImportedMaterialFunction )
+	{
+		return nullptr;
+	}
+
+	ImportContext.ImportedMaterialFunctions.Add( MaterialElement ) = ImportedMaterialFunction;
+
+	return ImportedMaterialFunction;
+}
+
+UMaterialFunction* FDatasmithImporter::FinalizeMaterialFunction(UObject* SourceMaterialFunction, const TCHAR* MaterialFunctionsFolderPath,
+	UMaterialFunction* ExistingMaterialFunction, TMap< UObject*, UObject* >* ReferencesToRemap)
+{
+	return Cast< UMaterialFunction >(DatasmithImporterImpl::FinalizeAsset(SourceMaterialFunction, MaterialFunctionsFolderPath, ExistingMaterialFunction, ReferencesToRemap));
+}
+
+UMaterialInterface* FDatasmithImporter::ImportMaterial( FDatasmithImportContext& ImportContext,
+	TSharedRef< IDatasmithBaseMaterialElement > MaterialElement, UMaterialInterface* ExistingMaterial )
+{
+	TRACE_CPUPROFILER_EVENT_SCOPE(FDatasmithImporter::ImportMaterial);
+
+	UMaterialInterface* ImportedMaterial = FDatasmithMaterialImporter::CreateMaterial( ImportContext, MaterialElement, ExistingMaterial );
+
+	if ( !ImportedMaterial )
+	{
+		return nullptr;
+	}
+
+#if MATERIAL_OPACITYMASK_DOESNT_SUPPORT_VIRTUALTEXTURE
+	TArray<UTexture*> OutOpacityMaskTextures;
+	if (ImportedMaterial->GetTexturesInPropertyChain(MP_OpacityMask, OutOpacityMaskTextures, nullptr, nullptr))
+	{
+		for (UTexture* CurrentTexture : OutOpacityMaskTextures)
+		{
+			UTexture2D* Texture2D = Cast<UTexture2D>(CurrentTexture);
+			if (Texture2D && Texture2D->VirtualTextureStreaming)
+			{
+				//Virtual textures are not supported yet in the OpacityMask slot, convert the texture back to a regular texture.
+				ImportContext.AssetsContext.VirtualTexturesToConvert.Add(Texture2D);
+			}
+		}
+	}
+#endif
+
+	UDatasmithAssetImportData* AssetImportData = Cast< UDatasmithAssetImportData >(ImportedMaterial->AssetImportData);
+
+	if (!AssetImportData)
+	{
+		AssetImportData = NewObject< UDatasmithAssetImportData >(ImportedMaterial);
+		ImportedMaterial->AssetImportData = AssetImportData;
+	}
+
+	AssetImportData->Update(ImportContext.Options->FilePath, ImportContext.FileHash.IsValid() ? &ImportContext.FileHash : nullptr);
+	AssetImportData->AssetImportOptions = ImportContext.Options->BaseOptions.AssetOptions;
+
+	// Record requirements on mesh building for this material
+	ImportContext.AssetsContext.MaterialsRequirements.Add( MaterialElement->GetName(), FDatasmithMaterialImporter::GetMaterialRequirements( ImportedMaterial ) );
+	ImportContext.ImportedMaterials.Add( MaterialElement ) = ImportedMaterial;
+
+	ImportMetaDataForObject( ImportContext, MaterialElement, ImportedMaterial );
+
+	return ImportedMaterial;
+}
+
+UObject* FDatasmithImporter::FinalizeMaterial( UObject* SourceMaterial, const TCHAR* MaterialsFolderPath, UMaterialInterface* ExistingMaterial, TMap< UObject*, UObject* >* ReferencesToRemap )
+{
+	TRACE_CPUPROFILER_EVENT_SCOPE(FDatasmithImporter::FinalizeMaterial);
+
+	using namespace DatasmithImporterImpl;
+
+	// Finalizing the master material might add a remapping for the instance parent property so make sure we have a remapping map available
+	TOptional< TMap< UObject*, UObject* > > ReferencesToRemapLocal;
+	if ( !ReferencesToRemap )
+	{
+		ReferencesToRemapLocal.Emplace();
+		ReferencesToRemap = &ReferencesToRemapLocal.GetValue();
+	}
+
+	if ( UMaterialInstance* SourceMaterialInstance = Cast< UMaterialInstance >( SourceMaterial ) )
+	{
+		if ( UMaterialInterface* SourceMaterialParent = SourceMaterialInstance->Parent )
+		{
+			FString SourceMaterialPath = SourceMaterialInstance->GetOutermost()->GetName();
+			FString SourceParentPath = SourceMaterialParent->GetOutermost()->GetName();
+
+			if ( SourceParentPath.StartsWith( SourceMaterialPath ) )
+			{
+				// Simply finalize the source parent material.
+				// Note that the parent material will be overridden on the existing material instance
+				FString DestinationParentPath = SourceParentPath;
+				DestinationParentPath.RemoveFromStart( SourceMaterialPath );
+				DestinationParentPath = MaterialsFolderPath / DestinationParentPath;
+
+				FinalizeMaterial( SourceMaterialParent, *DestinationParentPath, nullptr, ReferencesToRemap );
+			}
+		}
+	}
+
+	UMaterialEditingLibrary::DeleteAllMaterialExpressions( Cast< UMaterial >( ExistingMaterial ) );
+
+	UObject* DestinationMaterial = FinalizeAsset( SourceMaterial, MaterialsFolderPath, ExistingMaterial, ReferencesToRemap );
+
+	DatasmithImporterImpl::CompileMaterial( DestinationMaterial );
+
+	return DestinationMaterial;
+}
+
+void FDatasmithImporter::ImportActors( FDatasmithImportContext& ImportContext )
+{
+	TRACE_CPUPROFILER_EVENT_SCOPE(FDatasmithImporter::ImportActors);
+
+	/**
+	 * Hot fix for reimport issues UE-71655. A temporary created actor might have the same object path as the previously deleted actor.
+	 * This code below won't be needed when UE-71695 is fixed. This should be in 4.23.
+	 */
+	TArray< ADatasmithSceneActor* > SceneActors = FDatasmithImporterUtils::FindSceneActors( ImportContext.ActorsContext.FinalWorld, ImportContext.SceneAsset );
+	for ( ADatasmithSceneActor* SceneActor : SceneActors )
+	{
+		if ( !SceneActor )
+		{
+			continue;
+		}
+
+		if ( ImportContext.SceneAsset == SceneActor->Scene )
+		{
+			for ( TPair< FName, TSoftObjectPtr< AActor > >& Pair : SceneActor->RelatedActors )
+			{
+				// Try to load the actor. If we can't reset the soft object ptr
+				if ( !Pair.Value.LoadSynchronous() )
+				{
+					Pair.Value.Reset();
+				}
+			}
+		}
+	}
+	// end of the hotfix
+
+
+	ADatasmithSceneActor* ImportSceneActor = ImportContext.ActorsContext.ImportSceneActor;
+
+	// Create a scene actor to import with if we don't have one
+	if ( !ImportSceneActor )
+	{
+		// Create a the import scene actor for the import context
+		ImportSceneActor = FDatasmithImporterUtils::CreateImportSceneActor( ImportContext, FTransform::Identity );
+	}
+
+	const int32 ActorsCount = ImportContext.Scene->GetActorsCount();
+
+	FScopedSlowTask Progress( (float)ActorsCount, LOCTEXT("ImportActors", "Spawning actors..."), true, *ImportContext.Warn );
+	Progress.MakeDialog(true);
+
+	if ( ImportSceneActor )
+	{
+		ImportContext.Hierarchy.Push( ImportSceneActor->GetRootComponent() );
+
+		for (int32 i = 0; i < ActorsCount && !ImportContext.bUserCancelled; ++i)
+		{
+			ImportContext.bUserCancelled |= ImportContext.Warn->ReceivedUserCancel();
+
+			TSharedPtr< IDatasmithActorElement > ActorElement = ImportContext.Scene->GetActor(i);
+
+			if ( ActorElement.IsValid() )
+			{
+				Progress.EnterProgressFrame( 1.f, FText::FromString( FString::Printf( TEXT("Spawning actor %d/%d (%s) ..."), ( i + 1 ), ActorsCount, ActorElement->GetLabel() ) ) );
+
+				if ( ActorElement->IsAComponent() )
+				{
+					ImportActorAsComponent( ImportContext, ActorElement.ToSharedRef(), ImportSceneActor );
+				}
+				else
+				{
+					ImportActor( ImportContext, ActorElement.ToSharedRef() );
+				}
+			}
+		}
+
+		// Add all components under root actor to the root blueprint if Blueprint is required
+		if (ImportContext.Options->HierarchyHandling == EDatasmithImportHierarchy::UseOneBlueprint && ImportContext.RootBlueprint != nullptr)
+		{
+			// Reparent all scene components attached to root actor toward blueprint root
+			FKismetEditorUtilities::AddComponentsToBlueprint(ImportContext.RootBlueprint, ImportSceneActor->GetInstanceComponents(), false, (USCS_Node*)nullptr, true);
+		}
+
+		// After all actors were imported, perform a post import step so that any dependencies can be resolved
+		for (int32 i = 0; i < ActorsCount && !ImportContext.bUserCancelled; ++i)
+		{
+			ImportContext.bUserCancelled |= ImportContext.Warn->ReceivedUserCancel();
+
+			TSharedPtr< IDatasmithActorElement > ActorElement = ImportContext.Scene->GetActor(i);
+
+			if ( ActorElement.IsValid() && ActorElement->IsA( EDatasmithElementType::Camera ) )
+			{
+				FDatasmithCameraImporter::PostImportCameraActor( StaticCastSharedRef< IDatasmithCameraActorElement >( ActorElement.ToSharedRef() ), ImportContext );
+			}
+		}
+
+		ImportSceneActor->Scene = ImportContext.SceneAsset;
+
+		ImportContext.Hierarchy.Pop();
+	}
+
+	// Sky
+	if( ImportContext.Scene->GetUsePhysicalSky() )
+	{
+		AActor* SkyActor = FDatasmithLightImporter::CreatePhysicalSky(ImportContext);
+	}
+
+	if ( ImportContext.bUserCancelled )
+	{
+		DatasmithImporterImpl::DeleteImportSceneActorIfNeeded(ImportContext.ActorsContext, true );
+	}
+}
+
+AActor* FDatasmithImporter::ImportActor( FDatasmithImportContext& ImportContext, const TSharedRef< IDatasmithActorElement >& ActorElement )
+{
+	TRACE_CPUPROFILER_EVENT_SCOPE(FDatasmithImporter::ImportActor);
+
+	AActor* ImportedActor = nullptr;
+	if (ActorElement->IsA(EDatasmithElementType::HierarchicalInstanceStaticMesh))
+	{
+		TSharedRef< IDatasmithHierarchicalInstancedStaticMeshActorElement > HISMActorElement = StaticCastSharedRef< IDatasmithHierarchicalInstancedStaticMeshActorElement >( ActorElement );
+		ImportedActor =  FDatasmithActorImporter::ImportHierarchicalInstancedStaticMeshAsActor( ImportContext, HISMActorElement );
+	}
+	else if (ActorElement->IsA(EDatasmithElementType::StaticMeshActor))
+	{
+		TSharedRef< IDatasmithMeshActorElement > MeshActorElement = StaticCastSharedRef< IDatasmithMeshActorElement >( ActorElement );
+		ImportedActor = FDatasmithActorImporter::ImportStaticMeshActor( ImportContext, MeshActorElement );
+	}
+	else if (ActorElement->IsA(EDatasmithElementType::EnvironmentLight))
+	{
+		ImportedActor = FDatasmithActorImporter::ImportEnvironment( ImportContext, StaticCastSharedRef< IDatasmithEnvironmentElement >( ActorElement ) );
+	}
+	else if (ActorElement->IsA(EDatasmithElementType::Light))
+	{
+		ImportedActor = FDatasmithActorImporter::ImportLightActor( ImportContext, StaticCastSharedRef< IDatasmithLightActorElement >(ActorElement) );
+	}
+	else if (ActorElement->IsA(EDatasmithElementType::Camera))
+	{
+		ImportedActor = FDatasmithActorImporter::ImportCameraActor( ImportContext, StaticCastSharedRef< IDatasmithCameraActorElement >(ActorElement) );
+	}
+	else if (ActorElement->IsA(EDatasmithElementType::CustomActor))
+	{
+		ImportedActor = FDatasmithActorImporter::ImportCustomActor( ImportContext, StaticCastSharedRef< IDatasmithCustomActorElement >(ActorElement) );
+	}
+	else if (ActorElement->IsA(EDatasmithElementType::Landscape))
+	{
+		ImportedActor = FDatasmithActorImporter::ImportLandscapeActor( ImportContext, StaticCastSharedRef< IDatasmithLandscapeElement >(ActorElement) );
+	}
+	else if (ActorElement->IsA(EDatasmithElementType::PostProcessVolume))
+	{
+		ImportedActor = FDatasmithPostProcessImporter::ImportPostProcessVolume( StaticCastSharedRef< IDatasmithPostProcessVolumeElement >( ActorElement ), ImportContext, ImportContext.Options->OtherActorImportPolicy );
+	}
+	else
+	{
+		ImportedActor = FDatasmithActorImporter::ImportBaseActor( ImportContext, ActorElement );
+	}
+
+
+	if ( ImportedActor ) // It's possible that we didn't import an actor (ie: the user doesn't want to import the cameras), in that case, we'll skip it in the hierarchy
+	{
+		ImportContext.Hierarchy.Push( ImportedActor->GetRootComponent() );
+		ImportMetaDataForObject(ImportContext, ActorElement, ImportedActor);
+	}
+	else
+	{
+		ImportContext.ActorsContext.NonImportedDatasmithActors.Add( ActorElement->GetName() );
+	}
+
+	for (int32 i = 0; i < ActorElement->GetChildrenCount() && !ImportContext.bUserCancelled; ++i)
+	{
+		ImportContext.bUserCancelled |= ImportContext.Warn->ReceivedUserCancel();
+
+		const TSharedPtr< IDatasmithActorElement >& ChildActorElement = ActorElement->GetChild(i);
+
+		if ( ChildActorElement.IsValid() )
+		{
+			if ( ImportContext.Options->HierarchyHandling == EDatasmithImportHierarchy::UseMultipleActors && !ChildActorElement->IsAComponent() )
+			{
+				ImportActor( ImportContext, ChildActorElement.ToSharedRef() );
+			}
+			else if ( ImportedActor ) // Don't import the components of an actor that we didn't import
+			{
+				ImportActorAsComponent( ImportContext, ChildActorElement.ToSharedRef(), ImportedActor );
+			}
+		}
+	}
+
+	if ( ImportedActor )
+	{
+		ImportContext.Hierarchy.Pop();
+	}
+
+	return ImportedActor;
+}
+
+void FDatasmithImporter::ImportActorAsComponent(FDatasmithImportContext& ImportContext, const TSharedRef< IDatasmithActorElement >& ActorElement, AActor* InRootActor)
+{
+	if (!InRootActor)
+	{
+		return;
+	}
+
+	USceneComponent* SceneComponent = nullptr;
+
+	if (ActorElement->IsA(EDatasmithElementType::HierarchicalInstanceStaticMesh))
+	{
+		TSharedRef< IDatasmithHierarchicalInstancedStaticMeshActorElement > HierarchicalInstancedStaticMeshElement = StaticCastSharedRef< IDatasmithHierarchicalInstancedStaticMeshActorElement >(ActorElement);
+		SceneComponent = FDatasmithActorImporter::ImportHierarchicalInstancedStaticMeshComponent(ImportContext, HierarchicalInstancedStaticMeshElement, InRootActor);
+	}
+	else if (ActorElement->IsA(EDatasmithElementType::StaticMeshActor))
+	{
+		TSharedRef< IDatasmithMeshActorElement > MeshActorElement = StaticCastSharedRef< IDatasmithMeshActorElement >(ActorElement);
+		SceneComponent = FDatasmithActorImporter::ImportStaticMeshComponent(ImportContext, MeshActorElement, InRootActor);
+	}
+	else if (ActorElement->IsA(EDatasmithElementType::Light))
+	{
+		if (ImportContext.Options->LightImportPolicy == EDatasmithImportActorPolicy::Ignore)
+		{
+			return;
+		}
+
+		SceneComponent = FDatasmithLightImporter::ImportLightComponent(StaticCastSharedRef< IDatasmithLightActorElement >(ActorElement), ImportContext, InRootActor);
+	}
+	else if (ActorElement->IsA(EDatasmithElementType::Camera))
+	{
+		if (ImportContext.Options->CameraImportPolicy == EDatasmithImportActorPolicy::Ignore)
+		{
+			return;
+		}
+
+		SceneComponent = FDatasmithCameraImporter::ImportCineCameraComponent(StaticCastSharedRef< IDatasmithCameraActorElement >(ActorElement), ImportContext, InRootActor);
+	}
+	else
+	{
+		SceneComponent = FDatasmithActorImporter::ImportBaseActorAsComponent(ImportContext, ActorElement, InRootActor);
+	}
+
+	if (SceneComponent)
+	{
+		ImportContext.AddSceneComponent(SceneComponent->GetName(), SceneComponent);
+		ImportMetaDataForObject(ImportContext, ActorElement, SceneComponent);
+	}
+	else
+	{
+		ImportContext.ActorsContext.NonImportedDatasmithActors.Add(ActorElement->GetName());
+	}
+
+	for (int32 i = 0; i < ActorElement->GetChildrenCount(); ++i)
+	{
+		if (SceneComponent) // If we didn't import the current component, skip it in the hierarchy
+		{
+			ImportContext.Hierarchy.Push(SceneComponent);
+		}
+
+		ImportActorAsComponent(ImportContext, ActorElement->GetChild(i).ToSharedRef(), InRootActor);
+
+		if (SceneComponent)
+		{
+			ImportContext.Hierarchy.Pop();
+		}
+	}
+}
+
+void FDatasmithImporter::FinalizeActors( FDatasmithImportContext& ImportContext, TMap< UObject*, UObject* >* AssetReferencesToRemap )
+{
+	TRACE_CPUPROFILER_EVENT_SCOPE(FDatasmithImporter::FinalizeActors);
+
+	using namespace DatasmithImporterImpl;
+
+	if ( !ImportContext.bUserCancelled )
+	{
+		// Ensure a proper setup for the finalize of the actors
+		ADatasmithSceneActor*& ImportSceneActor = ImportContext.ActorsContext.ImportSceneActor;
+		TSet< ADatasmithSceneActor* >& FinalSceneActors = ImportContext.ActorsContext.FinalSceneActors;
+
+		if ( !ImportContext.ActorsContext.FinalWorld )
+		{
+			ImportContext.ActorsContext.FinalWorld = ImportContext.ActorsContext.ImportWorld;
+		}
+		else if ( !ImportContext.bIsAReimport && ImportSceneActor )
+		{
+				//Create a new datasmith scene actor in the final level
+				FActorSpawnParameters SpawnParameters;
+				SpawnParameters.Template = ImportSceneActor;
+				ADatasmithSceneActor* DestinationSceneActor = Cast< ADatasmithSceneActor >(ImportContext.ActorsContext.FinalWorld->SpawnActor< ADatasmithSceneActor >(SpawnParameters));
+
+				// Name new destination ADatasmithSceneActor to the DatasmithScene's name
+				DestinationSceneActor->SetActorLabel( ImportContext.Scene->GetName() );
+				DestinationSceneActor->MarkPackageDirty();
+				DestinationSceneActor->RelatedActors.Reset();
+
+				FinalSceneActors.Empty( 1 );
+				FinalSceneActors.Add( DestinationSceneActor );
+		}
+
+		if( FinalSceneActors.Num() == 0 )
+		{
+			if ( ImportContext.bIsAReimport )
+			{
+				FinalSceneActors.Append( FDatasmithImporterUtils::FindSceneActors( ImportContext.ActorsContext.FinalWorld, ImportContext.SceneAsset ) );
+				FinalSceneActors.Remove( ImportSceneActor );
+			}
+			else
+			{
+				FinalSceneActors.Add( ImportSceneActor );
+			}
+		}
+
+		for ( AActor* Actor : FinalSceneActors )
+		{
+			check(Actor->GetWorld() == ImportContext.ActorsContext.FinalWorld);
+		}
+
+
+		// Do the finalization for each actor from each FinalSceneActor
+		TMap< FSoftObjectPath, FSoftObjectPath > RenamedActorsMap;
+		TSet< FName > LayersUsedByActors;
+		const bool bShouldSpawnNonExistingActors = !ImportContext.bIsAReimport || ImportContext.Options->ReimportOptions.bRespawnDeletedActors;
+
+		for ( ADatasmithSceneActor* DestinationSceneActor : FinalSceneActors )
+		{
+			if ( !DestinationSceneActor )
+			{
+				continue;
+			}
+
+			if ( ImportSceneActor->Scene != DestinationSceneActor->Scene )
+			{
+				continue;
+			}
+
+			// In order to allow modification on components owned by DestinationSceneActor, unregister all of them
+			DestinationSceneActor->UnregisterAllComponents( /* bForReregister = */true);
+
+			ImportContext.ActorsContext.CurrentTargetedScene = DestinationSceneActor;
+
+			if ( ImportSceneActor != DestinationSceneActor )
+			{
+				// Before we delete the non imported actors, remove the old actor labels from the unique name provider
+				// as we don't care if the source labels clash with labels from actors that will be deleted or replaced on reimport
+				for (const TPair< FName, TSoftObjectPtr< AActor > >& ActorPair : DestinationSceneActor->RelatedActors)
+				{
+					if ( AActor* DestActor = ActorPair.Value.Get() )
+					{
+						ImportContext.ActorsContext.UniqueNameProvider.RemoveExistingName(DestActor->GetActorLabel());
+					}
+				}
+
+				FDatasmithImporterUtils::DeleteNonImportedDatasmithElementFromSceneActor( *ImportSceneActor, *DestinationSceneActor, ImportContext.ActorsContext.NonImportedDatasmithActors );
+			}
+
+			// Add Actor info to the remap info
+			TMap< UObject*, UObject* > PerSceneActorReferencesToRemap = AssetReferencesToRemap ? *AssetReferencesToRemap : TMap< UObject*, UObject* >();
+			PerSceneActorReferencesToRemap.Add( ImportSceneActor ) = DestinationSceneActor;
+			PerSceneActorReferencesToRemap.Add( ImportSceneActor->GetRootComponent() ) = DestinationSceneActor->GetRootComponent();
+
+			// #ueent_todo order of actors matters for ReferencesFix + re-parenting
+			for ( const TPair< FName, TSoftObjectPtr< AActor > >& SourceActorPair : ImportSceneActor->RelatedActors )
+			{
+				AActor* SourceActor = SourceActorPair.Value.Get();
+				if ( SourceActor == nullptr )
+				{
+					continue;
+				}
+
+				const bool bActorIsRelatedToDestionScene = DestinationSceneActor->RelatedActors.Contains( SourceActorPair.Key );
+				TSoftObjectPtr< AActor >& ExistingActorPtr = DestinationSceneActor->RelatedActors.FindOrAdd( SourceActorPair.Key );
+				const bool bShouldFinalizeActor = bShouldSpawnNonExistingActors || !bActorIsRelatedToDestionScene || ( ExistingActorPtr.Get() && !ExistingActorPtr.Get()->IsPendingKillPending() );
+
+				if ( bShouldFinalizeActor )
+				{
+					// Remember the original source path as FinalizeActor may set SourceActor's label, which apparently can also change its Name and package path
+					FSoftObjectPath OriginalSourcePath = FSoftObjectPath(SourceActor);
+					AActor* DestinationActor = FinalizeActor( ImportContext, *SourceActor, ExistingActorPtr.Get(), PerSceneActorReferencesToRemap );
+					RenamedActorsMap.Add(OriginalSourcePath, FSoftObjectPath(DestinationActor));
+					LayersUsedByActors.Append(DestinationActor->Layers);
+					ExistingActorPtr = DestinationActor;
+				}
+			}
+
+			for (  const TPair< FName, TSoftObjectPtr< AActor > >& DestinationActorPair : DestinationSceneActor->RelatedActors )
+			{
+				if ( AActor* Actor = DestinationActorPair.Value.Get() )
+				{
+					FixReferencesForObject( Actor, PerSceneActorReferencesToRemap );
+				}
+			}
+
+			// Modification is completed, re-register all components owned by DestinationSceneActor
+			DestinationSceneActor->RegisterAllComponents();
+		}
+
+		// Add the missing layers to the final world
+		FDatasmithImporterUtils::AddUniqueLayersToWorld( ImportContext.ActorsContext.FinalWorld, LayersUsedByActors );
+
+		// Fixed the soft object paths that were pointing to our pre-finalized actors.
+		TArray< UPackage* > PackagesToFix;
+
+		for ( const TPair< FName, TSoftObjectPtr< ULevelSequence > >& LevelSequencePair : ImportContext.SceneAsset->LevelSequences )
+		{
+			if (LevelSequencePair.Value)
+			{
+				PackagesToFix.Add( LevelSequencePair.Value->GetOutermost() );
+			}
+		}
+
+		for ( const TPair< FName, TSoftObjectPtr< ULevelVariantSets > >& LevelVariantSetsPair : ImportContext.SceneAsset->LevelVariantSets )
+		{
+			if (LevelVariantSetsPair.Value)
+			{
+				PackagesToFix.Add( LevelVariantSetsPair.Value->GetOutermost() );
+			}
+		}
+
+		FAssetToolsModule& AssetToolsModule = FModuleManager::GetModuleChecked< FAssetToolsModule >(TEXT("AssetTools"));
+		AssetToolsModule.Get().RenameReferencingSoftObjectPaths( PackagesToFix, RenamedActorsMap );
+	}
+
+	DeleteImportSceneActorIfNeeded( ImportContext.ActorsContext );
+
+	// Ensure layer visibility is properly updated for new actors associated with existing layers
+	ULayersSubsystem* LayersSubsystem = GEditor->GetEditorSubsystem<ULayersSubsystem>();
+	LayersSubsystem->UpdateAllActorsVisibility( false, true);
+
+	GEngine->BroadcastLevelActorListChanged();
+}
+
+AActor* FDatasmithImporter::FinalizeActor( FDatasmithImportContext& ImportContext, AActor& SourceActor, AActor* ExistingActor, TMap< UObject*, UObject* >& ReferencesToRemap )
+{
+	TRACE_CPUPROFILER_EVENT_SCOPE(FDatasmithImporter::FinalizeActor);
+
+	using namespace DatasmithImporterImpl;
+
+	// If the existing actor is not of the same class we destroy it
+	if ( ExistingActor && ExistingActor->GetClass() != SourceActor.GetClass() )
+	{
+		FDatasmithImporterUtils::DeleteActor( *ExistingActor );
+		ExistingActor = nullptr;
+	}
+
+	// Backup hierarchy
+	TArray< AActor* > Children;
+	if ( ExistingActor )
+	{
+		ExistingActor->GetAttachedActors( Children );
+
+		// In order to allow modification on components owned by ExistingActor, unregister all of them
+		ExistingActor->UnregisterAllComponents( /* bForReregister = */true );
+	}
+
+	AActor* DestinationActor = ExistingActor;
+	if ( !DestinationActor )
+	{
+		DestinationActor = ImportContext.ActorsContext.FinalWorld->SpawnActor( SourceActor.GetClass() );
+	}
+
+	// Update label to match the source actor's
+	DestinationActor->SetActorLabel( ImportContext.ActorsContext.UniqueNameProvider.GenerateUniqueName( SourceActor.GetActorLabel() ) );
+
+	check( DestinationActor );
+	ReferencesToRemap.Add( &SourceActor ) = DestinationActor;
+
+	TArray< FMigratedTemplatePairType > MigratedTemplates = MigrateTemplates(
+		&SourceActor,
+		ExistingActor,
+		&ReferencesToRemap,
+		true
+	);
+
+	// Copy actor data
+	{
+		TArray< uint8 > Bytes;
+		FActorWriter ObjectWriter( &SourceActor, Bytes );
+		FObjectReader ObjectReader( DestinationActor, Bytes );
+	}
+
+	FixReferencesForObject( DestinationActor, ReferencesToRemap );
+
+	DatasmithImporterImpl::FinalizeComponents( ImportContext, SourceActor, *DestinationActor, ReferencesToRemap );
+
+	// The templates for the actor need to be applied after the components were created.
+	ApplyMigratedTemplates( MigratedTemplates, DestinationActor );
+
+	// Restore hierarchy
+	for ( AActor* Child : Children )
+	{
+		Child->AttachToActor( DestinationActor, FAttachmentTransformRules::KeepWorldTransform );
+	}
+
+	// Hotfix for UE-69555
+	TInlineComponentArray<UHierarchicalInstancedStaticMeshComponent*> HierarchicalInstancedStaticMeshComponents;
+	DestinationActor->GetComponents(HierarchicalInstancedStaticMeshComponents);
+	for (UHierarchicalInstancedStaticMeshComponent* HierarchicalInstancedStaticMeshComponent : HierarchicalInstancedStaticMeshComponents )
+	{
+		HierarchicalInstancedStaticMeshComponent->BuildTreeIfOutdated( true, true );
+	}
+
+	if ( ALandscape* Landscape = Cast< ALandscape >( DestinationActor ) )
+	{
+		FPropertyChangedEvent MaterialPropertyChangedEvent( FindFieldChecked< UProperty >( Landscape->GetClass(), FName("LandscapeMaterial") ) );
+		Landscape->PostEditChangeProperty( MaterialPropertyChangedEvent );
+	}
+
+	DestinationActor->PostEditChange();
+	DestinationActor->RegisterAllComponents();
+
+	return DestinationActor;
+}
+
+void FDatasmithImporter::ImportLevelSequences( FDatasmithImportContext& ImportContext )
+{
+	const int32 SequencesCount = ImportContext.FilteredScene->GetLevelSequencesCount();
+	if ( !ImportContext.Options->BaseOptions.CanIncludeAnimation() || !ImportContext.Options->BaseOptions.bIncludeAnimation || SequencesCount == 0 )
+	{
+		return;
+	}
+
+	FScopedSlowTask Progress( (float)SequencesCount, LOCTEXT("ImportSlowTask", "Importing Level Sequences..."), true, *ImportContext.Warn );
+	Progress.MakeDialog(true);
+
+	// We can only parse a IDatasmithLevelSequenceElement with IDatasmithSubsequenceAnimationElements if their target
+	// subsequences' LevelSequenceElement have been parsed. We solve that with a structure we can repeatedly loop over,
+	// iteratively resolving all dependencies
+	TArray<TSharedPtr<IDatasmithLevelSequenceElement>> SequencesToImport;
+	SequencesToImport.Reserve(SequencesCount);
+	for ( int32 SequenceIndex = 0; SequenceIndex < SequencesCount && !ImportContext.bUserCancelled; ++SequenceIndex )
+	{
+		ImportContext.bUserCancelled |= ImportContext.Warn->ReceivedUserCancel();
+
+		TSharedPtr< IDatasmithLevelSequenceElement > SequenceElement = ImportContext.FilteredScene->GetLevelSequence( SequenceIndex );
+		if ( !SequenceElement )
+		{
+			continue;
+		}
+
+		SequencesToImport.Add(SequenceElement);
+	}
+
+	// If the scene is ok we will do at most HardLoopCounter passes
+	int32 HardLoopCounter = SequencesToImport.Num();
+	int32 NumImported = 0;
+	int32 LastNumImported = -1;
+	for (int32 IterationCounter = 0; IterationCounter < HardLoopCounter && !ImportContext.bUserCancelled; ++IterationCounter)
+	{
+		// Scan remaining sequences and import the ones we can, removing from this array
+		for ( int32 SequenceIndex = SequencesToImport.Num() - 1; SequenceIndex >= 0 && !ImportContext.bUserCancelled; --SequenceIndex )
+		{
+			ImportContext.bUserCancelled |= ImportContext.Warn->ReceivedUserCancel();
+
+			TSharedPtr<IDatasmithLevelSequenceElement>& SequenceElement = SequencesToImport[SequenceIndex];
+
+			if (!FDatasmithLevelSequenceImporter::CanImportLevelSequence(SequenceElement.ToSharedRef(), ImportContext))
+			{
+				continue;
+			}
+
+		ULevelSequence* ExistingLevelSequence = nullptr;
+		if ( ImportContext.SceneAsset )
+		{
+			TSoftObjectPtr< ULevelSequence >* ExistingLevelSequencePtr = ImportContext.SceneAsset->LevelSequences.Find( SequenceElement->GetName() );
+
+			if ( ExistingLevelSequencePtr )
+			{
+				ExistingLevelSequence = ExistingLevelSequencePtr->LoadSynchronous();
+			}
+		}
+
+			FString SequenceName = ObjectTools::SanitizeObjectName( SequenceElement->GetName() );
+			Progress.EnterProgressFrame( 1.f, FText::FromString( FString::Printf( TEXT("Importing level sequence %d/%d (%s) ..."), NumImported + 1, HardLoopCounter, *SequenceName ) ) );
+
+		ULevelSequence*& ImportedLevelSequence = ImportContext.ImportedLevelSequences.FindOrAdd( SequenceElement.ToSharedRef() );
+		if (ImportContext.SceneTranslator)
+		{
+			FDatasmithLevelSequencePayload LevelSequencePayload;
+			ImportContext.SceneTranslator->LoadLevelSequence(SequenceElement.ToSharedRef(), LevelSequencePayload);
+		}
+		ImportedLevelSequence = FDatasmithLevelSequenceImporter::ImportLevelSequence( SequenceElement.ToSharedRef(), ImportContext, ExistingLevelSequence );
+
+			SequencesToImport.RemoveAt(SequenceIndex);
+			++NumImported;
+		}
+
+		// If we do a full loop and haven't managed to parse at least one IDatasmithLevelSequenceElement, we'll assume something
+		// went wrong and step out.
+		if (NumImported == LastNumImported)
+		{
+			break;
+		}
+		LastNumImported = NumImported;
+	}
+
+	if (SequencesToImport.Num() > 0)
+	{
+		FString ErrorMessage = LOCTEXT("FailedToImport", "Failed to import some animation sequences:\n").ToString();
+		for (TSharedPtr<IDatasmithLevelSequenceElement>& Sequence: SequencesToImport)
+		{
+			ErrorMessage += FString(TEXT("\t")) + Sequence->GetName() + TEXT("\n");
+		}
+		ImportContext.LogError(FText::FromString(ErrorMessage));
+	}
+
+	// Assets have been imported and moved out of their import packages, clear them so that we don't look for them in there anymore
+	ImportContext.AssetsContext.LevelSequencesImportPackage.Reset();
+}
+
+ULevelSequence* FDatasmithImporter::FinalizeLevelSequence( ULevelSequence* SourceLevelSequence, const TCHAR* AnimationsFolderPath, ULevelSequence* ExistingLevelSequence )
+{
+	return Cast< ULevelSequence >( DatasmithImporterImpl::PublicizeAsset( SourceLevelSequence, AnimationsFolderPath, ExistingLevelSequence ) );
+}
+
+void FDatasmithImporter::ImportLevelVariantSets( FDatasmithImportContext& ImportContext )
+{
+	const int32 LevelVariantSetsCount = ImportContext.FilteredScene->GetLevelVariantSetsCount();
+	if ( LevelVariantSetsCount == 0 )
+	{
+		return;
+	}
+
+	FScopedSlowTask Progress( (float)LevelVariantSetsCount, LOCTEXT("ImportSlowTask", "Importing Level Variant Sets..."), true, *ImportContext.Warn );
+	Progress.MakeDialog(true);
+
+	for ( int32 LevelVariantSetIndex = 0; LevelVariantSetIndex < LevelVariantSetsCount && !ImportContext.bUserCancelled; ++LevelVariantSetIndex )
+	{
+		ImportContext.bUserCancelled |= ImportContext.Warn->ReceivedUserCancel();
+
+		TSharedPtr< IDatasmithLevelVariantSetsElement > LevelVariantSetsElement = ImportContext.FilteredScene->GetLevelVariantSets( LevelVariantSetIndex );
+		if ( !LevelVariantSetsElement )
+		{
+			continue;
+		}
+
+		ULevelVariantSets* ExistingLevelVariantSets = nullptr;
+		if ( ImportContext.SceneAsset )
+		{
+			TSoftObjectPtr< ULevelVariantSets >* ExistingLevelVariantSetsPtr = ImportContext.SceneAsset->LevelVariantSets.Find( LevelVariantSetsElement->GetName() );
+
+			if ( ExistingLevelVariantSetsPtr )
+			{
+				ExistingLevelVariantSets = ExistingLevelVariantSetsPtr->LoadSynchronous();
+			}
+		}
+
+		FString LevelVariantSetsName = ObjectTools::SanitizeObjectName( LevelVariantSetsElement->GetName() );
+		Progress.EnterProgressFrame( 1.f, FText::FromString( FString::Printf( TEXT("Importing level variant sets %d/%d (%s) ..."), LevelVariantSetIndex + 1, LevelVariantSetsCount, *LevelVariantSetsName ) ) );
+
+		ULevelVariantSets*& ImportedLevelVariantSets = ImportContext.ImportedLevelVariantSets.FindOrAdd( LevelVariantSetsElement.ToSharedRef() );
+		ImportedLevelVariantSets = FDatasmithLevelVariantSetsImporter::ImportLevelVariantSets( LevelVariantSetsElement.ToSharedRef(), ImportContext, ExistingLevelVariantSets );
+	}
+
+	// Assets have been imported and moved out of their import packages, clear them so that we don't look for them in there anymore
+	ImportContext.AssetsContext.LevelVariantSetsImportPackage.Reset();
+}
+
+ULevelVariantSets* FDatasmithImporter::FinalizeLevelVariantSets( ULevelVariantSets* SourceLevelVariantSets, const TCHAR* VariantsFolderPath, ULevelVariantSets* ExistingLevelVariantSets )
+{
+	return Cast< ULevelVariantSets >( DatasmithImporterImpl::PublicizeAsset( SourceLevelVariantSets, VariantsFolderPath, ExistingLevelVariantSets ) );
+}
+
+void FDatasmithImporter::ImportMetaDataForObject(FDatasmithImportContext& ImportContext, const TSharedRef<IDatasmithElement>& DatasmithElement, UObject* Object)
+{
+	if ( !Object )
+	{
+		return;
+	}
+
+	UDatasmithAssetUserData::FMetaDataContainer MetaData;
+
+	// Add Datasmith meta data
+	MetaData.Add( UDatasmithAssetUserData::UniqueIdMetaDataKey, DatasmithElement->GetName() );
+
+	// Check if there's metadata associated with the given element
+	const TSharedPtr<IDatasmithMetaDataElement>& MetaDataElement = ImportContext.Scene->GetMetaData( DatasmithElement );
+	if ( MetaDataElement.IsValid() )
+	{
+		int32 PropertiesCount = MetaDataElement->GetPropertiesCount();
+		MetaData.Reserve( PropertiesCount );
+		for ( int32 PropertyIndex = 0; PropertyIndex < PropertiesCount; ++PropertyIndex )
+		{
+			const TSharedPtr<IDatasmithKeyValueProperty>& Property = MetaDataElement->GetProperty( PropertyIndex );
+			MetaData.Add( Property->GetName(), Property->GetValue() );
+		}
+
+		MetaData.KeySort(FNameLexicalLess());
+	}
+
+	if ( MetaData.Num() > 0 )
+	{
+		// For AActor, the interface is actually implemented by the ActorComponent
+		AActor* Actor = Cast<AActor>( Object );
+		if ( Actor )
+		{
+			UActorComponent* ActorComponent = Actor->GetRootComponent();
+			if ( ActorComponent )
+			{
+				Object = ActorComponent;
+			}
+		}
+
+		if ( Object->GetClass()->ImplementsInterface( UInterface_AssetUserData::StaticClass() ) )
+		{
+			IInterface_AssetUserData* AssetUserData = Cast< IInterface_AssetUserData >( Object );
+
+			UDatasmithAssetUserData* DatasmithUserData = AssetUserData->GetAssetUserData< UDatasmithAssetUserData >();
+
+			if ( !DatasmithUserData )
+			{
+				DatasmithUserData = NewObject<UDatasmithAssetUserData>( Object, NAME_None, RF_Public | RF_Transactional );
+				AssetUserData->AddAssetUserData( DatasmithUserData );
+			}
+
+			DatasmithUserData->MetaData = MoveTemp( MetaData );
+		}
+	}
+}
+
+void FDatasmithImporter::FilterElementsToImport( FDatasmithImportContext& ImportContext )
+{
+	// Initialize the filtered scene as a copy of the original scene. We will use it to then filter out items to import.
+	ImportContext.FilteredScene = FDatasmithSceneFactory::DuplicateScene( ImportContext.Scene.ToSharedRef() );
+
+	// Filter meshes to import by consulting the AssetRegistry to see if that asset already exist
+	// or if it changed at all, if deemed identical filter the mesh out of the current import
+	IAssetRegistry& AssetRegistry = FModuleManager::LoadModuleChecked<FAssetRegistryModule>(TEXT("AssetRegistry")).Get();
+
+	// No Scene asset yet, all assets of the scene must be imported
+	if (!ImportContext.SceneAsset)
+	{
+		return;
+	}
+
+	auto ElementNeedsReimport = [&](const FString& FullyQualifiedName, TSharedRef<IDatasmithElement> Element, const FString& SourcePath) -> bool
+	{
+		const FAssetData AssetData = AssetRegistry.GetAssetByObjectPath(*FullyQualifiedName);
+		const FAssetDataTagMapSharedView::FFindTagResult ImportDataStr = AssetData.TagsAndValues.FindTag(TEXT("AssetImportData"));
+		FString CurrentRelativeFileName;
+
+		// Filter out Element only if it has valid and up to date import info
+		bool bImportThisElement = !ImportDataStr.IsSet();
+		if (!bImportThisElement)
+		{
+			TOptional<FAssetImportInfo> AssetImportInfo = FAssetImportInfo::FromJson(ImportDataStr.GetValue());
+			if (AssetImportInfo && AssetImportInfo->SourceFiles.Num() > 0)
+			{
+				FAssetImportInfo::FSourceFile ExistingSourceFile = AssetImportInfo->SourceFiles[0];
+				FMD5Hash ElementHash = Element->CalculateElementHash(false);
+				bImportThisElement = ExistingSourceFile.FileHash != ElementHash;
+				CurrentRelativeFileName = ExistingSourceFile.RelativeFilename;
+			}
+		}
+
+		// Sync import data now for skipped elements
+		if (!bImportThisElement && !SourcePath.IsEmpty())
+		{
+			FString ImportRelativeFileName = UAssetImportData::SanitizeImportFilename(SourcePath, AssetData.PackageName.ToString());
+			if (CurrentRelativeFileName != ImportRelativeFileName)
+			{
+				UObject* Asset = AssetData.GetAsset();
+				if (UAssetImportData* AssetImportData = Datasmith::GetAssetImportData(Asset))
+				{
+					AssetImportData->UpdateFilenameOnly(ImportRelativeFileName);
+				}
+			}
+		}
+
+		return bImportThisElement;
+	};
+
+	// Meshes part
+	ImportContext.FilteredScene->EmptyMeshes();
+	const TMap< FName, TSoftObjectPtr< UStaticMesh > >& StaticMeshes = ImportContext.SceneAsset->StaticMeshes;
+	for (int32 MeshIndex = 0; MeshIndex < ImportContext.Scene->GetMeshesCount(); ++MeshIndex)
+	{
+		TSharedRef< IDatasmithMeshElement > MeshElement = ImportContext.Scene->GetMesh( MeshIndex ).ToSharedRef();
+
+		bool bNeedsReimport = true;
+		FString AssetName = MeshElement->GetName();
+		if ( StaticMeshes.Contains( MeshElement->GetName() ) )
+		{
+			AssetName = StaticMeshes[ MeshElement->GetName() ].ToString();
+			bNeedsReimport = ElementNeedsReimport(AssetName, MeshElement, ImportContext.Options->FilePath );
+		}
+
+		if ( bNeedsReimport )
+		{
+			ImportContext.FilteredScene->AddMesh( MeshElement );
+		}
+		// If the mesh element does not need to be re-imported, register its name
+		else
+		{
+			const FAssetData AssetData = AssetRegistry.GetAssetByObjectPath( *AssetName);
+			ImportContext.AssetsContext.StaticMeshNameProvider.AddExistingName( FPaths::GetBaseFilename( AssetData.PackageName.ToString() ) );
+		}
+	}
+
+	// Texture part
+	ImportContext.FilteredScene->EmptyTextures();
+	const TMap< FName, TSoftObjectPtr< UTexture > >& Textures = ImportContext.SceneAsset->Textures;
+	for (int32 TexIndex = 0; TexIndex < ImportContext.Scene->GetTexturesCount(); ++TexIndex)
+	{
+		TSharedRef< IDatasmithTextureElement > TextureElement = ImportContext.Scene->GetTexture(TexIndex).ToSharedRef();
+
+		bool bNeedsReimport = true;
+		FString AssetName = TextureElement->GetName();
+		if ( Textures.Contains( TextureElement->GetName() ) )
+		{
+			AssetName = Textures[TextureElement->GetName()].ToString();
+			bNeedsReimport = ElementNeedsReimport(AssetName, TextureElement, ImportContext.Options->FilePath );
+		}
+
+		if ( bNeedsReimport )
+		{
+			ImportContext.FilteredScene->AddTexture( TextureElement );
+		}
+		// If the texture element does not need to be re-imported, register its name
+		else
+		{
+			const FAssetData AssetData = AssetRegistry.GetAssetByObjectPath( *AssetName);
+			ImportContext.AssetsContext.TextureNameProvider.AddExistingName( FPaths::GetBaseFilename( AssetData.PackageName.ToString() ) );
+		}
+	}
+}
+
+void FDatasmithImporter::FinalizeImport(FDatasmithImportContext& ImportContext, const TSet<UObject*>& ValidAssets)
+{
+	TRACE_CPUPROFILER_EVENT_SCOPE(FDatasmithImporter::FinalizeImport);
+
+	const int32 NumImportedObjects = ImportContext.ImportedStaticMeshes.Num() +
+									 ImportContext.ImportedTextures.Num() +
+									 ImportContext.ImportedMaterialFunctions.Num() +
+									 ImportContext.ImportedMaterials.Num() +
+									 ImportContext.ImportedLevelSequences.Num() +
+									 ImportContext.ImportedLevelVariantSets.Num();
+	const int32 NumAssetsToFinalize = ValidAssets.Num() == 0 ? NumImportedObjects : ValidAssets.Num() + ImportContext.ImportedLevelSequences.Num() + ImportContext.ImportedLevelVariantSets.Num();
+	const int32 NumStaticMeshToBuild = ImportContext.ImportedStaticMeshes.Num();
+
+	FScopedSlowTask Progress((float)NumAssetsToFinalize + NumStaticMeshToBuild, LOCTEXT("FinalizingAssets", "Finalizing Assets"), true, *ImportContext.Warn);
+	Progress.MakeDialog(true);
+
+	TMap<UObject*, UObject*> ReferencesToRemap;
+
+	// Array of packages containing templates which are referring to assets as TSoftObjectPtr or FSoftObjectPath
+	TArray<UPackage*> PackagesToCheck;
+
+	int32 AssetIndex = 0;
+
+	const FString& RootFolderPath = ImportContext.AssetsContext.RootFolderPath;
+	const FString& TransientFolderPath = ImportContext.AssetsContext.TransientFolderPath;
+
+	// Needs to be done in dependencies order (textures -> materials -> static meshes)
+	for (const TPair< TSharedRef< IDatasmithTextureElement >, UTexture* >& ImportedTexturePair : ImportContext.ImportedTextures)
+	{
+		if (ImportContext.bUserCancelled)
+		{
+			break;
+		}
+
+		UTexture* SourceTexture = ImportedTexturePair.Value;
+
+		if (!SourceTexture || (ValidAssets.Num() > 0 && !ValidAssets.Contains(Cast<UObject>(SourceTexture))))
+		{
+			continue;
+		}
+
+		FName TextureId = ImportedTexturePair.Key->GetName();
+
+		Progress.EnterProgressFrame(1.f, FText::FromString(FString::Printf(TEXT("Finalizing assets %d/%d (Texture %s) ..."), ++AssetIndex, NumAssetsToFinalize, *SourceTexture->GetName())));
+
+		TSoftObjectPtr< UTexture >& ExistingTexturePtr = ImportContext.SceneAsset->Textures.FindOrAdd(TextureId);
+		UTexture* ExistingTexture = ExistingTexturePtr.Get();
+
+		FString SourcePackagePath = SourceTexture->GetOutermost()->GetName();
+		FString DestinationPackagePath = SourcePackagePath.Replace(*TransientFolderPath, *RootFolderPath, ESearchCase::CaseSensitive);
+
+		ExistingTexturePtr = FinalizeTexture(SourceTexture, *DestinationPackagePath, ExistingTexture, &ReferencesToRemap);
+		DatasmithImporterImpl::CheckAssetPersistenceValidity(ExistingTexturePtr.Get(), ImportContext);
+	}
+
+	// Unregister all actors component to avoid excessive refresh in the 3D engine while updating materials.
+	for (TObjectIterator<AActor> ActorIterator; ActorIterator; ++ActorIterator)
+	{
+		if (ActorIterator->GetWorld())
+		{
+			ActorIterator->UnregisterAllComponents( /* bForReregister = */true);
+		}
+	}
+
+	for (const TPair< TSharedRef< IDatasmithBaseMaterialElement >, UMaterialFunction* >& ImportedMaterialFunctionPair : ImportContext.ImportedMaterialFunctions)
+	{
+		if (ImportContext.bUserCancelled)
+		{
+			break;
+		}
+
+		UMaterialFunction* SourceMaterialFunction = ImportedMaterialFunctionPair.Value;
+
+		if (!SourceMaterialFunction || (ValidAssets.Num() > 0 && !ValidAssets.Contains(Cast<UObject>(SourceMaterialFunction))))
+		{
+			continue;
+		}
+
+		FName MaterialFunctionId = ImportedMaterialFunctionPair.Key->GetName();
+
+		Progress.EnterProgressFrame(1.f, FText::FromString(FString::Printf(TEXT("Finalizing assets %d/%d (Material Function %s) ..."), ++AssetIndex, NumAssetsToFinalize, *SourceMaterialFunction->GetName())));
+
+		TSoftObjectPtr< UMaterialFunction >& ExistingMaterialFunctionPtr = ImportContext.SceneAsset->MaterialFunctions.FindOrAdd(MaterialFunctionId);
+		UMaterialFunction* ExistingMaterialFunction = ExistingMaterialFunctionPtr.Get();
+
+		FString SourcePackagePath = SourceMaterialFunction->GetOutermost()->GetName();
+		FString DestinationPackagePath = SourcePackagePath.Replace(*TransientFolderPath, *RootFolderPath, ESearchCase::CaseSensitive);
+
+		ExistingMaterialFunctionPtr = FinalizeMaterialFunction(SourceMaterialFunction, *DestinationPackagePath, ExistingMaterialFunction, &ReferencesToRemap);
+		DatasmithImporterImpl::CheckAssetPersistenceValidity(ExistingMaterialFunctionPtr.Get(), ImportContext);
+	}
+
+	for (const TPair< TSharedRef< IDatasmithBaseMaterialElement >, UMaterialInterface* >& ImportedMaterialPair : ImportContext.ImportedMaterials)
+	{
+		if (ImportContext.bUserCancelled)
+		{
+			break;
+		}
+
+		UMaterialInterface* SourceMaterialInterface = ImportedMaterialPair.Value;
+
+		if (!SourceMaterialInterface || (ValidAssets.Num() > 0 && !ValidAssets.Contains(Cast<UObject>(SourceMaterialInterface))))
+		{
+			continue;
+		}
+
+		FName MaterialId = ImportedMaterialPair.Key->GetName();
+
+		Progress.EnterProgressFrame(1.f, FText::FromString(FString::Printf(TEXT("Finalizing assets %d/%d (Material %s) ..."), ++AssetIndex, NumAssetsToFinalize, *SourceMaterialInterface->GetName())));
+
+		TSoftObjectPtr< UMaterialInterface >& ExistingMaterialPtr = ImportContext.SceneAsset->Materials.FindOrAdd(MaterialId);
+		UMaterialInterface* ExistingMaterial = ExistingMaterialPtr.Get();
+
+		FString SourcePackagePath = SourceMaterialInterface->GetOutermost()->GetName();
+		FString DestinationPackagePath = SourcePackagePath.Replace( *TransientFolderPath, *RootFolderPath, ESearchCase::CaseSensitive );
+
+		if (UMaterial* SourceMaterial = Cast< UMaterial >(SourceMaterialInterface))
+		{
+			SourceMaterial->RebuildExpressionTextureReferences();
+
+			for (FMaterialFunctionInfo& MaterialFunctionInfo : SourceMaterial->MaterialFunctionInfos)
+			{
+				if (MaterialFunctionInfo.Function && MaterialFunctionInfo.Function->GetOutermost() == SourceMaterial->GetOutermost())
+				{
+					FinalizeMaterial(MaterialFunctionInfo.Function, *DestinationPackagePath, nullptr, &ReferencesToRemap);
+				}
+			}
+		}
+
+		ExistingMaterialPtr = FinalizeMaterial(SourceMaterialInterface, *DestinationPackagePath, ExistingMaterial, &ReferencesToRemap);
+
+		// Add material to array of packages to apply soft object path redirection to
+		if (ExistingMaterialPtr.IsValid())
+		{
+			PackagesToCheck.Add(ExistingMaterialPtr->GetOutermost());
+			DatasmithImporterImpl::CheckAssetPersistenceValidity(ExistingMaterialPtr.Get(), ImportContext);
+		}
+	}
+
+	DatasmithImporterImpl::ConvertUnsupportedVirtualTexture(ImportContext, ImportContext.AssetsContext.VirtualTexturesToConvert, ReferencesToRemap);
+
+	// Materials have been updated, we can register everything back.
+	for (TObjectIterator<AActor> ActorIterator; ActorIterator; ++ActorIterator)
+	{
+		if (ActorIterator->GetWorld())
+		{
+			ActorIterator->RegisterAllComponents();
+		}
+	}
+
+	// Sometimes, the data is invalid and we get the same UStaticMesh multiple times
+	TSet< UStaticMesh* > StaticMeshes;
+	for (const TPair< TSharedRef< IDatasmithMeshElement >, UStaticMesh* >& ImportedStaticMeshPair : ImportContext.ImportedStaticMeshes)
+	{
+		if (ImportContext.bUserCancelled)
+		{
+			break;
+		}
+
+		UStaticMesh* SourceStaticMesh = ImportedStaticMeshPair.Value;
+
+		if (!SourceStaticMesh || (ValidAssets.Num() > 0 && !ValidAssets.Contains(Cast<UObject>(SourceStaticMesh))))
+		{
+			continue;
+		}
+
+		FName StaticMeshId = ImportedStaticMeshPair.Key->GetName();
+
+		Progress.EnterProgressFrame(1.f, FText::FromString(FString::Printf(TEXT("Finalizing assets %d/%d (Static Mesh %s) ..."), ++AssetIndex, NumAssetsToFinalize, *SourceStaticMesh->GetName())));
+
+		TSoftObjectPtr< UStaticMesh >& ExistingStaticMeshPtr = ImportContext.SceneAsset->StaticMeshes.FindOrAdd(StaticMeshId);
+		UStaticMesh* ExistingStaticMesh = ExistingStaticMeshPtr.Get();
+
+		FString SourcePackagePath = SourceStaticMesh->GetOutermost()->GetName();
+		FString DestinationPackagePath = SourcePackagePath.Replace( *TransientFolderPath, *RootFolderPath, ESearchCase::CaseSensitive );
+
+		ExistingStaticMeshPtr = FinalizeStaticMesh(SourceStaticMesh, *DestinationPackagePath, ExistingStaticMesh, &ReferencesToRemap, false);
+		DatasmithImporterImpl::CheckAssetPersistenceValidity(ExistingStaticMeshPtr.Get(), ImportContext);
+
+		StaticMeshes.Add(ExistingStaticMeshPtr.Get());
+	}
+
+	int32 StaticMeshIndex = 0;
+	auto ProgressFunction = [&](UStaticMesh* StaticMesh)
+	{
+		Progress.EnterProgressFrame(1.f, FText::FromString(FString::Printf(TEXT("Building Static Mesh %d/%d (%s) ..."), ++StaticMeshIndex, StaticMeshes.Num(), *StaticMesh->GetName())));
+		return !ImportContext.bUserCancelled;
+	};
+
+	FDatasmithStaticMeshImporter::BuildStaticMeshes(StaticMeshes.Array(), ProgressFunction);
+
+	for (const TPair< TSharedRef< IDatasmithLevelSequenceElement >, ULevelSequence* >& ImportedLevelSequencePair : ImportContext.ImportedLevelSequences)
+	{
+		if (ImportContext.bUserCancelled)
+		{
+			break;
+		}
+
+		ULevelSequence* SourceLevelSequence = ImportedLevelSequencePair.Value;
+
+		if (!SourceLevelSequence)
+		{
+			continue;
+		}
+
+		FName LevelSequenceId = ImportedLevelSequencePair.Key->GetName();
+
+		Progress.EnterProgressFrame(1.f, FText::FromString(FString::Printf(TEXT("Finalizing assets %d/%d (Level Sequence %s) ..."), ++AssetIndex, NumAssetsToFinalize, *SourceLevelSequence->GetName())));
+
+		TSoftObjectPtr< ULevelSequence >& ExistingLevelSequencePtr = ImportContext.SceneAsset->LevelSequences.FindOrAdd(LevelSequenceId);
+		ULevelSequence* ExistingLevelSequence = ExistingLevelSequencePtr.Get();
+
+		FString SourcePackagePath = SourceLevelSequence->GetOutermost()->GetName();
+		FString DestinationPackagePath = SourcePackagePath.Replace( *TransientFolderPath, *RootFolderPath, ESearchCase::CaseSensitive );
+
+		ExistingLevelSequencePtr = FinalizeLevelSequence(SourceLevelSequence, *DestinationPackagePath, ExistingLevelSequence);
+		ImportContext.SceneAsset->RegisterPreWorldRenameCallback();
+	}
+
+	for (const TPair< TSharedRef< IDatasmithLevelVariantSetsElement >, ULevelVariantSets* >& ImportedLevelVariantSetsPair : ImportContext.ImportedLevelVariantSets)
+	{
+		if (ImportContext.bUserCancelled)
+		{
+			break;
+		}
+
+		ULevelVariantSets* SourceLevelVariantSets = ImportedLevelVariantSetsPair.Value;
+
+		if (!SourceLevelVariantSets)
+		{
+			continue;
+		}
+
+		FName LevelVariantSetsId = ImportedLevelVariantSetsPair.Key->GetName();
+
+		Progress.EnterProgressFrame(1.f, FText::FromString(FString::Printf(TEXT("Finalizing assets %d/%d (Level Variant Sets %s) ..."), ++AssetIndex, NumAssetsToFinalize, *SourceLevelVariantSets->GetName())));
+
+		TSoftObjectPtr< ULevelVariantSets >& ExistingLevelVariantSetsPtr = ImportContext.SceneAsset->LevelVariantSets.FindOrAdd(LevelVariantSetsId);
+		ULevelVariantSets* ExistingLevelVariantSets = ExistingLevelVariantSetsPtr.Get();
+
+		FString SourcePackagePath = SourceLevelVariantSets->GetOutermost()->GetName();
+		FString DestinationPackagePath = SourcePackagePath.Replace( *TransientFolderPath, *RootFolderPath, ESearchCase::CaseSensitive );
+
+		ExistingLevelVariantSetsPtr = FinalizeLevelVariantSets(SourceLevelVariantSets, *DestinationPackagePath, ExistingLevelVariantSets);
+		ImportContext.SceneAsset->RegisterPreWorldRenameCallback();
+	}
+
+	// Apply soft object path redirection to identified packages
+	if (PackagesToCheck.Num() > 0 && ReferencesToRemap.Num() > 0)
+	{
+		TMap<FSoftObjectPath, FSoftObjectPath> AssetRedirectorMap;
+
+		for ( const TPair< UObject*, UObject* >& ReferenceToRemap : ReferencesToRemap )
+		{
+			AssetRedirectorMap.Emplace( ReferenceToRemap.Key ) = ReferenceToRemap.Value;
+		}
+
+		IAssetTools& AssetTools = FModuleManager::LoadModuleChecked<FAssetToolsModule>("AssetTools").Get();
+		AssetTools.RenameReferencingSoftObjectPaths(PackagesToCheck, AssetRedirectorMap);
+	}
+
+	if ( ImportContext.ShouldImportActors() )
+	{
+		FinalizeActors(ImportContext, &ReferencesToRemap);
+	}
+
+	// Everything has been finalized, make sure the UDatasmithScene is set to dirty
+	if (ImportContext.SceneAsset)
+	{
+		ImportContext.SceneAsset->MarkPackageDirty();
+	}
+
+	FGlobalComponentReregisterContext RecreateComponents;
+
+	// Flush the transaction buffer (eg. avoid corrupted hierarchies after repeated undo actions)
+	// This is an aggressive workaround while we don't properly support undo history.
+	GEditor->ResetTransaction(LOCTEXT("Reset Transaction Buffer", "Datasmith Import Finalization"));
+}
+
+
+#undef LOCTEXT_NAMESPACE
