Copying //UE4/Dev-Enterprise @ cl 9420543 to Dev-Main (//UE4/Dev-Main)

diff --git a/Engine/Plugins/Enterprise/DatasmithImporter/Source/DatasmithImporter/Private/DatasmithImporterEditorSettings.cpp b/Engine/Plugins/Enterprise/DatasmithImporter/Source/DatasmithImporter/Private/DatasmithImporterEditorSettings.cpp
new file mode 100644
--- /dev/null
+++ b/Engine/Plugins/Enterprise/DatasmithImporter/Source/DatasmithImporter/Private/DatasmithImporterEditorSettings.cpp
@@ -0,0 +1,8 @@
+// Copyright 1998-2019 Epic Games, Inc. All Rights Reserved.
+
+#include "DatasmithImporterEditorSettings.h"
+
+UDatasmithImporterEditorSettings::UDatasmithImporterEditorSettings()
+	: bOfflineImporter(false)
+{
+}
\ No newline at end of file
diff --git a/Engine/Plugins/Enterprise/DatasmithImporter/Source/DatasmithImporter/Private/DatasmithLandscapeImporter.cpp b/Engine/Plugins/Enterprise/DatasmithImporter/Source/DatasmithImporter/Private/DatasmithLandscapeImporter.cpp
new file mode 100644
--- /dev/null
+++ b/Engine/Plugins/Enterprise/DatasmithImporter/Source/DatasmithImporter/Private/DatasmithLandscapeImporter.cpp
@@ -0,0 +1,127 @@
+// Copyright 1998-2019 Epic Games, Inc. All Rights Reserved.
+
+#include "DatasmithLandscapeImporter.h"
+
+#include "DatasmithActorImporter.h"
+#include "DatasmithImportContext.h"
+#include "IDatasmithSceneElements.h"
+#include "ObjectTemplates/DatasmithLandscapeTemplate.h"
+
+#include "Landscape.h"
+#include "LandscapeDataAccess.h"
+#include "LandscapeEditorModule.h"
+#include "LandscapeEditorObject.h"
+#include "LandscapeEditorUtils.h"
+#include "LandscapeFileFormatInterface.h"
+#include "NewLandscapeUtils.h"
+
+#include "Modules/ModuleManager.h"
+#include "Utility/DatasmithImporterUtils.h"
+
+#define LOCTEXT_NAMESPACE "DatasmithImportFactory"
+
+AActor* FDatasmithLandscapeImporter::ImportLandscapeActor( const TSharedRef< IDatasmithLandscapeElement >& LandscapeActorElement, FDatasmithImportContext& ImportContext, EDatasmithImportActorPolicy ImportActorPolicy )
+{
+	TStrongObjectPtr< ULandscapeEditorObject > LandscapeEditorObject( NewObject< ULandscapeEditorObject >() );
+
+	LandscapeEditorObject->ImportLandscape_HeightmapFilename = LandscapeActorElement->GetHeightmap();
+	LandscapeEditorObject->NewLandscape_Scale = LandscapeActorElement->GetScale();
+
+	TArray< FLandscapeFileResolution > ImportResolutions;
+	FNewLandscapeUtils::ImportLandscapeData( LandscapeEditorObject.Get(), ImportResolutions );
+
+	const int32 ComponentCountX = LandscapeEditorObject->NewLandscape_ComponentCount.X;
+	const int32 ComponentCountY = LandscapeEditorObject->NewLandscape_ComponentCount.Y;
+	const int32 QuadsPerComponent = LandscapeEditorObject->NewLandscape_SectionsPerComponent * LandscapeEditorObject->NewLandscape_QuadsPerSection;
+	const int32 SizeX = ComponentCountX * QuadsPerComponent + 1;
+	const int32 SizeY = ComponentCountY * QuadsPerComponent + 1;
+
+	TOptional< TArray< FLandscapeImportLayerInfo > > ImportLayers = FNewLandscapeUtils::CreateImportLayersInfo( LandscapeEditorObject.Get(), ENewLandscapePreviewMode::ImportLandscape );
+
+	if ( !ImportLayers )
+	{
+		return nullptr;
+	}
+
+	TArray<uint16> HeightData = FNewLandscapeUtils::ComputeHeightData( LandscapeEditorObject.Get(), ImportLayers.GetValue(), ENewLandscapePreviewMode::ImportLandscape );
+
+	const ELandscapeImportAlphamapType ImportLandscape_AlphamapType = ELandscapeImportAlphamapType::Additive;
+
+	const FVector Offset = FTransform( LandscapeActorElement->GetRotation(), FVector::ZeroVector, 
+		LandscapeActorElement->GetScale() ).TransformVector( FVector( -ComponentCountX * QuadsPerComponent / 2, -ComponentCountY * QuadsPerComponent / 2, 0 ) );
+
+	FVector OriginalTranslation = LandscapeActorElement->GetTranslation();
+	FVector OriginalScale = LandscapeActorElement->GetScale();
+
+	LandscapeActorElement->SetTranslation( LandscapeActorElement->GetTranslation() + Offset );
+	LandscapeActorElement->SetScale( LandscapeActorElement->GetScale() );
+
+	ALandscape* Landscape = Cast< ALandscape >(FDatasmithActorImporter::ImportActor( ALandscape::StaticClass(), LandscapeActorElement, ImportContext, ImportActorPolicy,
+		[ &HeightData, &ImportLayers, SizeX, SizeY, LandscapeEditorObject, ImportLandscape_AlphamapType, ActorScale = LandscapeActorElement->GetScale() ]( AActor* NewActor )
+	{
+		check( Cast< ALandscape >( NewActor ) );
+
+		NewActor->SetActorRelativeScale3D( ActorScale );
+
+		TMap<FGuid, TArray<uint16>> HeightmapDataPerLayers;
+		HeightmapDataPerLayers.Add(FGuid(), HeightData);
+		TMap<FGuid, TArray<FLandscapeImportLayerInfo>> MaterialLayerDataPerLayer;
+		MaterialLayerDataPerLayer.Add(FGuid(), ImportLayers.GetValue());
+
+		Cast< ALandscape >( NewActor )->Import( FGuid::NewGuid(), 0, 0, SizeX - 1, SizeY - 1, LandscapeEditorObject->NewLandscape_SectionsPerComponent, LandscapeEditorObject->NewLandscape_QuadsPerSection,
+			HeightmapDataPerLayers, nullptr, MaterialLayerDataPerLayer, ImportLandscape_AlphamapType );
+	} ) );
+
+	LandscapeActorElement->SetTranslation( OriginalTranslation );
+	LandscapeActorElement->SetScale( OriginalScale );
+
+	if ( !Landscape )
+	{
+		return nullptr;
+	}
+
+	UDatasmithLandscapeTemplate* LandscapeTemplate = NewObject< UDatasmithLandscapeTemplate >( Landscape->GetRootComponent() );
+	LandscapeTemplate->LandscapeMaterial = FDatasmithImporterUtils::FindAsset< UMaterialInterface >( ImportContext.AssetsContext, LandscapeActorElement->GetMaterial() );
+
+	// automatically calculate a lighting LOD that won't crash lightmass (hopefully)
+	// < 2048x2048 -> LOD0
+	// >=2048x2048 -> LOD1
+	// >= 4096x4096 -> LOD2
+	// >= 8192x8192 -> LOD3
+	LandscapeTemplate->StaticLightingLOD = FMath::DivideAndRoundUp(FMath::CeilLogTwo((SizeX * SizeY) / (2048 * 2048) + 1), (uint32)2);
+	
+	LandscapeTemplate->Apply( Landscape );
+
+	Landscape->ReimportHeightmapFilePath = LandscapeEditorObject->ImportLandscape_HeightmapFilename;
+
+	ULandscapeInfo* LandscapeInfo = Landscape->CreateLandscapeInfo();
+	LandscapeInfo->UpdateLayerInfoMap(Landscape);
+
+	// Import doesn't fill in the LayerInfo for layers with no data, do that now
+	const TArray< FLandscapeImportLayer >& ImportLandscapeLayersList = LandscapeEditorObject->ImportLandscape_Layers;
+	for(int32 i = 0; i < ImportLandscapeLayersList.Num(); i++)
+	{
+		if(ImportLandscapeLayersList[i].LayerInfo != nullptr)
+		{
+			Landscape->EditorLayerSettings.Add(FLandscapeEditorLayerSettings(ImportLandscapeLayersList[i].LayerInfo, ImportLandscapeLayersList[i].SourceFilePath));
+
+			int32 LayerInfoIndex = LandscapeInfo->GetLayerInfoIndex(ImportLandscapeLayersList[i].LayerName);
+			if(ensure(LayerInfoIndex != INDEX_NONE))
+			{
+				FLandscapeInfoLayerSettings& LayerSettings = LandscapeInfo->Layers[LayerInfoIndex];
+				LayerSettings.LayerInfoObj = ImportLandscapeLayersList[i].LayerInfo;
+			}
+		}
+	}
+
+	Landscape->RegisterAllComponents();
+
+	// Need to explicitly call PostEditChange on the LandscapeMaterial property or the landscape proxy won't update its material
+	FPropertyChangedEvent MaterialPropertyChangedEvent( FindFieldChecked< UProperty >( Landscape->GetClass(), FName("LandscapeMaterial") ) );
+	Landscape->PostEditChangeProperty( MaterialPropertyChangedEvent );
+	Landscape->PostEditChange();
+
+	return Landscape;
+}
+
+#undef LOCTEXT_NAMESPACE
diff --git a/Engine/Plugins/Enterprise/DatasmithImporter/Source/DatasmithImporter/Private/DatasmithLightImporter.cpp b/Engine/Plugins/Enterprise/DatasmithImporter/Source/DatasmithImporter/Private/DatasmithLightImporter.cpp
new file mode 100644
--- /dev/null
+++ b/Engine/Plugins/Enterprise/DatasmithImporter/Source/DatasmithImporter/Private/DatasmithLightImporter.cpp
@@ -0,0 +1,631 @@
+// Copyright 1998-2019 Epic Games, Inc. All Rights Reserved.
+
+#include "DatasmithLightImporter.h"
+
+#include "DatasmithActorImporter.h"
+#include "DatasmithAreaLightActor.h"
+#include "DatasmithImportContext.h"
+#include "DatasmithImportOptions.h"
+#include "DatasmithMaterialExpressions.h"
+#include "DatasmithSceneFactory.h"
+
+#include "ObjectTemplates/DatasmithAreaLightActorTemplate.h"
+#include "ObjectTemplates/DatasmithLightComponentTemplate.h"
+#include "ObjectTemplates/DatasmithPointLightComponentTemplate.h"
+#include "ObjectTemplates/DatasmithSceneComponentTemplate.h"
+#include "ObjectTemplates/DatasmithSkyLightComponentTemplate.h"
+
+#include "AssetRegistryModule.h"
+#include "ComponentReregisterContext.h"
+#include "EngineUtils.h"
+#include "Internationalization/Internationalization.h"
+#include "ObjectTools.h"
+#include "RawMesh.h"
+
+#include "Components/ChildActorComponent.h"
+#include "Components/DirectionalLightComponent.h"
+#include "Components/LightComponent.h"
+#include "Components/LightmassPortalComponent.h"
+#include "Components/PointLightComponent.h"
+#include "Components/SkyLightComponent.h"
+#include "Components/SpotLightComponent.h"
+#include "Components/StaticMeshComponent.h"
+
+#include "Editor/GroupActor.h"
+
+#include "Engine/DirectionalLight.h"
+#include "Engine/Light.h"
+#include "Engine/PointLight.h"
+#include "Engine/SkyLight.h"
+#include "Engine/SpotLight.h"
+#include "Engine/StaticMesh.h"
+#include "Engine/StaticMeshActor.h"
+#include "Engine/TextureCube.h"
+#include "Engine/TextureLightProfile.h"
+
+#include "Lightmass/LightmassPortal.h"
+
+#include "Materials/Material.h"
+#include "Misc/MessageDialog.h"
+#include "Misc/Paths.h"
+
+#define LOCTEXT_NAMESPACE "DatasmithLightImporter"
+
+namespace
+{
+	UClass* GetActorClassForLightActorElement( const TSharedRef< IDatasmithLightActorElement >& LightElement )
+	{
+		UClass* LightActorClass = APointLight::StaticClass();
+
+		if ( LightElement->IsA(EDatasmithElementType::AreaLight) )
+		{
+			LightActorClass = AActor::StaticClass();
+		}
+		else if (LightElement->IsA(EDatasmithElementType::LightmassPortal))
+		{
+			LightActorClass = ALightmassPortal::StaticClass();
+		}
+		else if (LightElement->IsA(EDatasmithElementType::DirectionalLight))
+		{
+			LightActorClass = ADirectionalLight::StaticClass();
+		}
+		else if ( LightElement->IsA( EDatasmithElementType::SpotLight ) )
+		{
+			LightActorClass = ASpotLight::StaticClass();
+		}
+		else if ( LightElement->IsA( EDatasmithElementType::PointLight ) )
+		{
+			LightActorClass = APointLight::StaticClass();
+		}
+
+		return LightActorClass;
+	}
+
+	EDatasmithAreaLightActorType GetLightActorTypeForLightType( const EDatasmithAreaLightType LightType )
+	{
+		EDatasmithAreaLightActorType LightActorType = EDatasmithAreaLightActorType::Point;
+
+		switch ( LightType )
+		{
+		case EDatasmithAreaLightType::Spot:
+			LightActorType = EDatasmithAreaLightActorType::Spot;
+			break;
+
+		case EDatasmithAreaLightType::Point:
+			LightActorType = EDatasmithAreaLightActorType::Point;
+			break;
+
+		case EDatasmithAreaLightType::IES_DEPRECATED:
+			LightActorType = EDatasmithAreaLightActorType::Point;
+			break;
+
+		case EDatasmithAreaLightType::Rect:
+			LightActorType = EDatasmithAreaLightActorType::Rect;
+			break;
+		}
+
+		return LightActorType;
+	}
+}
+
+AActor* FDatasmithLightImporter::ImportLightActor( const TSharedRef< IDatasmithLightActorElement >& LightElement, FDatasmithImportContext& ImportContext )
+{
+	if ( ImportContext.Options->LightImportPolicy == EDatasmithImportActorPolicy::Ignore )
+	{
+		return nullptr;
+	}
+
+	CreateIESTexture( ImportContext, LightElement );
+
+	AActor* ImportedLightActor = nullptr;
+	UClass* LightActorClass = GetActorClassForLightActorElement( LightElement );
+
+	if ( LightElement->IsA( EDatasmithElementType::AreaLight ) )
+	{
+		TSharedRef< IDatasmithAreaLightElement > AreaLightElement = StaticCastSharedRef< IDatasmithAreaLightElement >( LightElement );
+
+		ImportedLightActor = ImportAreaLightActor( AreaLightElement, ImportContext );
+
+		// In the case of a re-import, the returned value will be null if area light was deleted. Abort import
+		if (ImportedLightActor == nullptr)
+		{
+			return nullptr;
+		}
+	}
+	else
+	{
+		FQuat LightRotation = LightElement->GetRotation();
+
+		if ( LightElement->GetUseIes() ) // For IES lights that are not area lights, the IES rotation should be baked into the light transform
+		{
+			FQuat IesLightRotation = LightElement->GetRotation() * LightElement->GetIesRotation();
+			LightElement->SetRotation( IesLightRotation );
+		}
+
+		AActor* Actor = FDatasmithActorImporter::ImportActor( LightActorClass, LightElement, ImportContext, ImportContext.Options->LightImportPolicy );
+		if ( !Actor )
+		{
+			return nullptr;
+		}
+
+		if ( LightElement->IsA( EDatasmithElementType::LightmassPortal ) )
+		{
+			ALightmassPortal* LightmassPortal = Cast< ALightmassPortal >( Actor );
+			ImportedLightActor = LightmassPortal;
+		}
+		else if ( LightElement->IsA( EDatasmithElementType::DirectionalLight ) )
+		{
+			ADirectionalLight* DirectionalLight = Cast< ADirectionalLight >( Actor );
+
+			if ( !DirectionalLight )
+			{
+				return nullptr;
+			}
+
+			SetupLightComponent( DirectionalLight->GetLightComponent(), LightElement, *ImportContext.AssetsContext.MaterialsFinalPackage->GetPathName(), *ImportContext.AssetsContext.LightPackage->GetPathName() );
+
+			ImportedLightActor = DirectionalLight;
+		}
+		else if ( LightElement->IsA( EDatasmithElementType::SpotLight ) )
+		{
+			ASpotLight* SpotLight = Cast< ASpotLight >(Actor);
+
+			if ( !SpotLight )
+			{
+				return nullptr;
+			}
+
+			SetupSpotLightComponent( Cast< USpotLightComponent >( SpotLight->GetLightComponent() ), StaticCastSharedRef< IDatasmithSpotLightElement >( LightElement ),
+				*ImportContext.AssetsContext.LightPackage->GetPathName(), *ImportContext.AssetsContext.MaterialsFinalPackage->GetPathName() );
+
+			ImportedLightActor = SpotLight;
+		}
+		else if ( LightElement->IsA( EDatasmithElementType::PointLight ) )
+		{
+			APointLight* PointLight = Cast< APointLight >(Actor);
+
+			if ( !PointLight )
+			{
+				return nullptr;
+			}
+
+			SetupPointLightComponent( Cast< UPointLightComponent >( PointLight->GetLightComponent() ), StaticCastSharedRef< IDatasmithPointLightElement >( LightElement ),
+				*ImportContext.AssetsContext.LightPackage->GetPathName(), *ImportContext.AssetsContext.MaterialsFinalPackage->GetPathName() );
+
+			ImportedLightActor = PointLight;
+		}
+	}
+
+	if ( ImportedLightActor && ImportedLightActor->GetRootComponent() && !ImportedLightActor->GetRootComponent()->IsRegistered() )
+	{
+		ImportedLightActor->GetRootComponent()->RegisterComponent();
+	}
+
+	return ImportedLightActor;
+}
+
+USceneComponent* FDatasmithLightImporter::ImportLightComponent( const TSharedRef< IDatasmithLightActorElement >& LightElement, FDatasmithImportContext& ImportContext, UObject* Outer )
+{
+	CreateIESTexture( ImportContext, LightElement );
+
+	USceneComponent* LightComponent = nullptr;
+
+	if ( LightElement->IsA(EDatasmithElementType::AreaLight) )
+	{
+		TSharedRef< IDatasmithAreaLightElement > AreaLightElement = StaticCastSharedRef< IDatasmithAreaLightElement >( LightElement );
+		LightComponent = ImportAreaLightComponent( AreaLightElement, ImportContext, Outer );
+	}
+	else if ( LightElement->IsA( EDatasmithElementType::LightmassPortal ) )
+	{
+		TSharedRef< IDatasmithLightmassPortalElement > LightmassPortalElement = StaticCastSharedRef< IDatasmithLightmassPortalElement >( LightElement );
+		LightComponent = ImportLightmassPortalComponent( LightmassPortalElement, ImportContext, Outer );
+	}
+	else if ( LightElement->IsA( EDatasmithElementType::DirectionalLight ) )
+	{
+		USceneComponent* SceneComponent = FDatasmithActorImporter::ImportSceneComponent( UDirectionalLightComponent::StaticClass(), LightElement, ImportContext, Outer );
+
+		UDirectionalLightComponent* DirectionalLightComponent = Cast< UDirectionalLightComponent >( SceneComponent );
+
+		if ( !DirectionalLightComponent )
+		{
+			return nullptr;
+		}
+
+		SetupLightComponent( DirectionalLightComponent, LightElement, *ImportContext.AssetsContext.MaterialsFinalPackage->GetPathName(), *ImportContext.AssetsContext.LightPackage->GetPathName() );
+
+		LightComponent = DirectionalLightComponent;
+	}
+	else if ( LightElement->IsA( EDatasmithElementType::SpotLight ) || LightElement->IsA( EDatasmithElementType::PointLight ) )
+	{
+		FQuat LightRotation = LightElement->GetRotation();
+
+		if ( LightElement->GetUseIes() ) // For IES lights that are not area lights, the IES rotation should be baked into the light transform
+		{
+			FQuat IesLightRotation = LightElement->GetRotation() * LightElement->GetIesRotation();
+			LightElement->SetRotation( IesLightRotation );
+		}
+
+		if ( LightElement->IsA( EDatasmithElementType::SpotLight ) )
+		{
+			USceneComponent* SceneComponent = FDatasmithActorImporter::ImportSceneComponent( USpotLightComponent::StaticClass(), LightElement, ImportContext, Outer );
+
+			USpotLightComponent* SpotLightComponent = Cast< USpotLightComponent >( SceneComponent );
+
+			if ( !SpotLightComponent )
+			{
+				return nullptr;
+			}
+
+			TSharedRef< IDatasmithSpotLightElement > SpotLightElement = StaticCastSharedRef< IDatasmithSpotLightElement >( LightElement );
+			SetupSpotLightComponent( SpotLightComponent, SpotLightElement, *ImportContext.AssetsContext.LightPackage->GetPathName(), *ImportContext.AssetsContext.MaterialsFinalPackage->GetPathName() );
+
+			LightComponent = SpotLightComponent;
+		}
+		else if ( LightElement->IsA( EDatasmithElementType::PointLight ) )
+		{
+			USceneComponent* SceneComponent = FDatasmithActorImporter::ImportSceneComponent( UPointLightComponent::StaticClass(), LightElement, ImportContext, Outer );
+
+			UPointLightComponent* PointLightComponent = Cast< UPointLightComponent >( SceneComponent );
+
+			if ( !PointLightComponent )
+			{
+				return nullptr;
+			}
+
+			SetupPointLightComponent( PointLightComponent, StaticCastSharedRef< IDatasmithPointLightElement >( LightElement ),
+				*ImportContext.AssetsContext.LightPackage->GetPathName(), *ImportContext.AssetsContext.MaterialsFinalPackage->GetPathName() );
+
+			LightComponent = PointLightComponent;
+		}
+
+		LightElement->SetRotation( LightRotation );
+	}
+
+	if ( LightComponent )
+	{
+		LightComponent->RegisterComponent();
+	}
+
+	return LightComponent;
+}
+
+void FDatasmithLightImporter::SetTextureLightProfile( const TSharedRef< IDatasmithLightActorElement >& LightElement, UDatasmithLightComponentTemplate* LightComponentTemplate, const TCHAR* LightsFolderPath )
+{
+	if ( !LightComponentTemplate )
+	{
+		return;
+	}
+
+	if ( LightElement->GetUseIes() )
+	{
+		LightComponentTemplate->bUseIESBrightness = LightElement->GetUseIesBrightness();
+		LightComponentTemplate->IESBrightnessScale = LightElement->GetIesBrightnessScale();
+
+		if ( FCString::Strlen( LightElement->GetIesFile() ) > 0 && FPaths::FileExists( LightElement->GetIesFile() ) )
+		{
+			UTextureLightProfile* LightProfile = FindTextureLightProfile( LightElement, LightsFolderPath );
+
+			if ( LightProfile )
+			{
+				LightComponentTemplate->IESTexture = LightProfile;
+			}
+		}
+	}
+}
+
+UTextureLightProfile* FDatasmithLightImporter::FindTextureLightProfile( const TSharedRef< IDatasmithLightActorElement >& LightElement, const TCHAR* LightsFolderPath )
+{
+	UTextureLightProfile* IESTexture = nullptr;
+
+	if ( FCString::Strlen( LightElement->GetIesFile() ) > 0 && FPaths::FileExists( LightElement->GetIesFile() ) )
+	{
+		FString TextureName = FPaths::GetBaseFilename(FString(LightElement->GetIesFile())) + TEXT("_IES");
+		TextureName = ObjectTools::SanitizeObjectName(TextureName);
+
+		FSoftObjectPath LightProfileObjectPath( FPaths::Combine( LightsFolderPath, TextureName ) );
+
+		IESTexture = Cast< UTextureLightProfile >( LightProfileObjectPath.TryLoad() );
+	}
+
+	return IESTexture;
+}
+
+void FDatasmithLightImporter::CreateIESTexture(FDatasmithImportContext& InContext, const TSharedPtr< IDatasmithLightActorElement >& InLightElement)
+{
+	if (InLightElement->GetUseIes())
+	{
+		FString IESFilename(InLightElement->GetIesFile());
+		if (!IESFilename.IsEmpty() && FPaths::FileExists(IESFilename))
+		{
+			FString IesName = FPaths::GetBaseFilename(IESFilename) + TEXT("_IES");
+			IesName = ObjectTools::SanitizeObjectName(IesName);
+
+			if ( !InContext.ParsedIesFiles.Contains( IesName ) )
+			{
+				InContext.ParsedIesFiles.Add( IesName );
+
+				FDatasmithMaterialExpressions::CreateDatasmithIES( InLightElement->GetIesFile(), InContext.AssetsContext.LightPackage.Get(), InContext.ObjectFlags );
+			}
+		}
+	}
+}
+
+void FDatasmithLightImporter::SetupLightComponent( ULightComponent* LightComponent, const TSharedPtr< IDatasmithLightActorElement >& LightElement, const TCHAR* MaterialsFolderPath, const TCHAR* LightsFolderPath )
+{
+	if ( !LightComponent )
+	{
+		return;
+	}
+
+	UDatasmithLightComponentTemplate* LightComponentTemplate = NewObject< UDatasmithLightComponentTemplate >( LightComponent );
+
+	LightComponentTemplate->bVisible = LightElement->IsEnabled();
+	LightComponentTemplate->Intensity = LightElement->GetIntensity();
+	LightComponentTemplate->CastShadows = true;
+	LightComponentTemplate->LightColor = LightElement->GetColor().ToFColor( true );
+	LightComponentTemplate->bUseTemperature = LightElement->GetUseTemperature();
+	LightComponentTemplate->Temperature = LightElement->GetTemperature();
+
+	if ( LightElement->GetLightFunctionMaterial().IsValid() )
+	{
+		FString BaseName = LightElement->GetLightFunctionMaterial()->GetName();
+		FString MaterialName = FPaths::Combine( MaterialsFolderPath, BaseName + TEXT(".") + BaseName );
+		UMaterialInterface* Material = Cast< UMaterialInterface >( FSoftObjectPath( *MaterialName ).TryLoad() );
+
+		if ( Material )
+		{
+			LightComponentTemplate->LightFunctionMaterial = Material;
+		}
+	}
+
+	if ( LightComponent->IsA< UPointLightComponent >() )
+	{
+		SetTextureLightProfile( LightElement.ToSharedRef(), LightComponentTemplate, LightsFolderPath );
+	}
+
+	LightComponentTemplate->Apply( LightComponent );
+	LightComponent->UpdateColorAndBrightness();
+}
+
+void FDatasmithLightImporter::SetupPointLightComponent( UPointLightComponent* PointLightComponent, const TSharedRef< IDatasmithPointLightElement >& PointLightElement, const TCHAR* LightsFolderPath, const TCHAR* MaterialsFolderPath )
+{
+	if ( !PointLightComponent )
+	{
+		return;
+	}
+
+	SetupLightComponent( PointLightComponent, PointLightElement, MaterialsFolderPath, LightsFolderPath );
+
+	UDatasmithPointLightComponentTemplate* PointLightComponentTemplate = NewObject< UDatasmithPointLightComponentTemplate >( PointLightComponent );
+
+	switch ( PointLightElement->GetIntensityUnits() )
+	{
+	case EDatasmithLightUnits::Candelas:
+		PointLightComponentTemplate->IntensityUnits = ELightUnits::Candelas;
+		break;
+	case EDatasmithLightUnits::Lumens:
+		PointLightComponentTemplate->IntensityUnits = ELightUnits::Lumens;
+		break;
+	default:
+		PointLightComponentTemplate->IntensityUnits = ELightUnits::Unitless;
+		break;
+	}
+
+	if ( PointLightElement->GetSourceRadius() > 0.f )
+	{
+		PointLightComponentTemplate->SourceRadius = PointLightElement->GetSourceRadius();
+	}
+
+	if ( PointLightElement->GetSourceLength() > 0.f )
+	{
+		PointLightComponentTemplate->SourceLength = PointLightElement->GetSourceLength();
+	}
+
+	if ( PointLightElement->GetAttenuationRadius() > 0.f )
+	{
+		PointLightComponentTemplate->AttenuationRadius = PointLightElement->GetAttenuationRadius();
+	}
+
+	PointLightComponentTemplate->Apply( PointLightComponent );
+	PointLightComponent->UpdateColorAndBrightness();
+}
+
+void FDatasmithLightImporter::SetupSpotLightComponent( USpotLightComponent* SpotLightComponent, const TSharedRef< IDatasmithSpotLightElement >& SpotLightElement, const TCHAR* LightsFolderPath, const TCHAR* MaterialsFolderPath )
+{
+	SetupPointLightComponent( SpotLightComponent, SpotLightElement, LightsFolderPath, MaterialsFolderPath );
+
+	SpotLightComponent->InnerConeAngle = SpotLightElement->GetInnerConeAngle();
+	SpotLightComponent->OuterConeAngle = SpotLightElement->GetOuterConeAngle();
+}
+
+AActor* FDatasmithLightImporter::ImportAreaLightActor( const TSharedRef< IDatasmithAreaLightElement >& AreaLightElement, FDatasmithImportContext& ImportContext )
+{
+	AActor* AreaLightActor = CreateAreaLightActor( AreaLightElement, ImportContext );
+
+	return AreaLightActor;
+}
+
+USceneComponent* FDatasmithLightImporter::ImportAreaLightComponent( const TSharedRef< IDatasmithAreaLightElement >& AreaLightElement, FDatasmithImportContext& ImportContext, UObject* Outer )
+{
+	USceneComponent* MainComponent = nullptr;
+
+	FSoftObjectPath LightShapeBlueprintRef = FSoftObjectPath( TEXT("/DatasmithContent/Datasmith/DatasmithArealight.DatasmithArealight") );
+	UBlueprint* LightShapeBlueprint = Cast< UBlueprint >( LightShapeBlueprintRef.TryLoad() );
+
+	if ( LightShapeBlueprint )
+	{
+		UChildActorComponent* ChildActorComponent = Cast< UChildActorComponent >( FDatasmithActorImporter::ImportSceneComponent( UChildActorComponent::StaticClass(), AreaLightElement, ImportContext, Outer ) );
+		ChildActorComponent->SetChildActorClass( TSubclassOf< AActor > ( LightShapeBlueprint->GeneratedClass ) );
+		ChildActorComponent->CreateChildActor();
+
+		ADatasmithAreaLightActor* LightShapeActor = Cast< ADatasmithAreaLightActor >( ChildActorComponent->GetChildActor() );
+
+		if ( LightShapeActor )
+		{
+			LightShapeActor->SetActorLabel( AreaLightElement->GetLabel() );
+			SetupAreaLightActor( AreaLightElement, ImportContext, LightShapeActor );
+
+			MainComponent = ChildActorComponent;
+		}
+	}
+
+	return MainComponent;
+}
+
+ULightmassPortalComponent* FDatasmithLightImporter::ImportLightmassPortalComponent( const TSharedRef< IDatasmithLightmassPortalElement >& LightElement, FDatasmithImportContext& ImportContext, UObject* Outer )
+{
+	USceneComponent* SceneComponent = FDatasmithActorImporter::ImportSceneComponent( ULightmassPortalComponent::StaticClass(), LightElement, ImportContext, Outer );
+
+	SceneComponent->RegisterComponent();
+
+	return Cast< ULightmassPortalComponent >( SceneComponent );
+}
+
+AActor* FDatasmithLightImporter::CreateAreaLightActor( const TSharedRef< IDatasmithAreaLightElement >& AreaLightElement, FDatasmithImportContext& ImportContext )
+{
+	FSoftObjectPath LightShapeBlueprintRef = FSoftObjectPath( TEXT("/DatasmithContent/Datasmith/DatasmithArealight.DatasmithArealight") );
+	UBlueprint* LightShapeBlueprint = Cast< UBlueprint >( LightShapeBlueprintRef.TryLoad() );
+
+	if ( !LightShapeBlueprint )
+	{
+		return nullptr;
+	}
+
+	AActor* Actor = FDatasmithActorImporter::ImportActor( LightShapeBlueprint->GeneratedClass, AreaLightElement, ImportContext, ImportContext.Options->LightImportPolicy );
+
+	ADatasmithAreaLightActor* LightShapeActor = Cast< ADatasmithAreaLightActor >( Actor );
+	SetupAreaLightActor( AreaLightElement, ImportContext, LightShapeActor );
+
+	return LightShapeActor;
+}
+
+void FDatasmithLightImporter::SetupAreaLightActor( const TSharedRef< IDatasmithAreaLightElement >& AreaLightElement, FDatasmithImportContext& ImportContext, ADatasmithAreaLightActor* LightShapeActor )
+{
+	if ( !LightShapeActor )
+	{
+		return;
+	}
+
+	UDatasmithAreaLightActorTemplate* LightActorTemplate = NewObject< UDatasmithAreaLightActorTemplate >( LightShapeActor );
+
+	LightShapeActor->UnregisterAllComponents( true );
+
+	LightActorTemplate->LightType = GetLightActorTypeForLightType( AreaLightElement->GetLightType() );
+	LightActorTemplate->LightShape = (EDatasmithAreaLightActorShape)AreaLightElement->GetLightShape();
+	LightActorTemplate->Dimensions = FVector2D( AreaLightElement->GetLength(), AreaLightElement->GetWidth() );
+	LightActorTemplate->Color = AreaLightElement->GetColor();
+	LightActorTemplate->Intensity = AreaLightElement->GetIntensity();
+	LightActorTemplate->IntensityUnits = (ELightUnits)AreaLightElement->GetIntensityUnits();
+
+	if ( AreaLightElement->GetUseTemperature() )
+	{
+		LightActorTemplate->Temperature = AreaLightElement->GetTemperature();
+	}
+
+	if ( AreaLightElement->GetUseIes() )
+	{
+		LightActorTemplate->IESTexture = FindTextureLightProfile( AreaLightElement, *ImportContext.AssetsContext.LightPackage->GetPathName() );
+		LightActorTemplate->bUseIESBrightness = AreaLightElement->GetUseIesBrightness();
+		LightActorTemplate->IESBrightnessScale = AreaLightElement->GetIesBrightnessScale();
+		LightActorTemplate->Rotation = AreaLightElement->GetIesRotation().Rotator();
+	}
+
+	if ( AreaLightElement->GetSourceRadius() > 0.f )
+	{
+		LightActorTemplate->SourceRadius = AreaLightElement->GetSourceRadius();
+	}
+
+	if ( AreaLightElement->GetSourceLength() > 0.f )
+	{
+		LightActorTemplate->SourceLength = AreaLightElement->GetSourceLength();
+	}
+
+	if ( AreaLightElement->GetAttenuationRadius() > 0.f )
+	{
+		LightActorTemplate->AttenuationRadius = AreaLightElement->GetAttenuationRadius();
+	}
+
+	LightActorTemplate->Apply( LightShapeActor );
+
+	LightShapeActor->RegisterAllComponents();
+
+	LightShapeActor->RerunConstructionScripts();
+}
+
+AActor* FDatasmithLightImporter::CreateSkyLight(const TSharedPtr< IDatasmithShaderElement >& ShaderElement, FDatasmithImportContext& ImportContext, bool bUseHDRMat)
+{
+	if ( ImportContext.Options->LightImportPolicy == EDatasmithImportActorPolicy::Ignore )
+	{
+		return nullptr;
+	}
+
+	TSharedRef< IDatasmithDirectionalLightElement > SkyLightElement = FDatasmithSceneFactory::CreateDirectionalLight( TEXT("SkyLight") );
+
+	if ( bUseHDRMat )
+	{
+		FQuat ActorRotation = FQuat::MakeFromEuler( FVector(0.f, 0.f, ShaderElement->GetEmitTextureSampler().Rotation * 360.f) );
+		SkyLightElement->SetRotation( ActorRotation );
+	}
+
+	AActor* Actor = FDatasmithActorImporter::ImportActor( ASkyLight::StaticClass(), SkyLightElement, ImportContext, ImportContext.Options->LightImportPolicy );
+
+	ASkyLight* SkyLight = Cast< ASkyLight >( Actor );
+
+	if (SkyLight == nullptr)
+	{
+		ImportContext.LogError( FText::Format( LOCTEXT( "CreateSkyLight", "Failed to create the sky light {0}" ), FText::FromString( ShaderElement->GetLabel() ) ) );
+
+		return nullptr;
+	}
+
+	UDatasmithSceneComponentTemplate* SceneComponentTemplate = NewObject< UDatasmithSceneComponentTemplate >( SkyLight );
+	SceneComponentTemplate->Mobility = EComponentMobility::Static;
+
+	UDatasmithSkyLightComponentTemplate* SkyLightComponentTemplate = NewObject< UDatasmithSkyLightComponentTemplate >( SkyLight );
+
+	if (bUseHDRMat)
+	{
+		FString EmitTexturePath = ShaderElement->GetEmitTexture();
+
+		if ( !EmitTexturePath.IsEmpty() && FPaths::IsRelative( EmitTexturePath ) )
+		{
+			EmitTexturePath = FPaths::Combine( ImportContext.AssetsContext.TexturesFinalPackage->GetPathName(), ShaderElement->GetEmitTexture() );
+		}
+
+		UTextureCube* CubeTexture = Cast< UTextureCube >( FSoftObjectPath( EmitTexturePath ).TryLoad() );
+
+		if (CubeTexture)
+		{
+			SkyLightComponentTemplate->SourceType = ESkyLightSourceType::SLS_SpecifiedCubemap;
+			SkyLightComponentTemplate->CubemapResolution = 512;
+			SkyLightComponentTemplate->Cubemap = CubeTexture;
+		}
+	}
+	else
+	{
+		SkyLightComponentTemplate->Cubemap = nullptr;
+	}
+
+	SceneComponentTemplate->Apply( SkyLight->GetLightComponent() );
+	SkyLightComponentTemplate->Apply( SkyLight->GetLightComponent() );
+
+	SkyLight->GetLightComponent()->RegisterComponent();
+
+	SkyLight->MarkComponentsRenderStateDirty();
+	SkyLight->MarkPackageDirty();
+
+	return Actor;
+}
+
+AActor* FDatasmithLightImporter::CreateHDRISkyLight(const TSharedPtr< IDatasmithShaderElement >& ShaderElement, FDatasmithImportContext& ImportContext)
+{
+	return CreateSkyLight(ShaderElement, ImportContext, true);
+}
+
+AActor* FDatasmithLightImporter::CreatePhysicalSky(FDatasmithImportContext& ImportContext)
+{
+	TSharedPtr< IDatasmithShaderElement > ShaderElement = FDatasmithSceneFactory::CreateShader(TEXT("voiddummymat"));
+	return CreateSkyLight(ShaderElement, ImportContext, false);
+}
+
+#undef LOCTEXT_NAMESPACE
diff --git a/Engine/Plugins/Enterprise/DatasmithImporter/Source/DatasmithImporter/Private/DatasmithMaterialImporter.cpp b/Engine/Plugins/Enterprise/DatasmithImporter/Source/DatasmithImporter/Private/DatasmithMaterialImporter.cpp
new file mode 100644
--- /dev/null
+++ b/Engine/Plugins/Enterprise/DatasmithImporter/Source/DatasmithImporter/Private/DatasmithMaterialImporter.cpp
@@ -0,0 +1,374 @@
+// Copyright 1998-2019 Epic Games, Inc. All Rights Reserved.
+
+#include "DatasmithMaterialImporter.h"
+
+#include "DatasmithImporterModule.h"
+#include "DatasmithMaterialElements.h"
+#include "DatasmithMaterialExpressions.h"
+
+#include "MasterMaterials/DatasmithMasterMaterial.h"
+#include "MasterMaterials/DatasmithMasterMaterialManager.h"
+#include "MasterMaterials/DatasmithMasterMaterialSelector.h"
+#include "ObjectTemplates/DatasmithMaterialInstanceTemplate.h"
+#include "Utility/DatasmithImporterUtils.h"
+
+#include "AssetRegistryModule.h"
+
+#include "Engine/Texture.h"
+#include "Engine/Texture2D.h"
+
+#include "Materials/Material.h"
+#include "Materials/MaterialFunction.h"
+#include "Materials/MaterialInstanceConstant.h"
+
+#include "ObjectTools.h"
+
+
+namespace DatasmithMaterialImporterUtils
+{
+	int32 ComputeMaterialExpressionHash( IDatasmithMaterialExpression* MaterialExpression );
+
+	int32 ComputeExpressionInputHash( IDatasmithExpressionInput* ExpressionInput )
+	{
+		if ( !ExpressionInput )
+		{
+			return 0;
+		}
+
+		uint32 Hash = 0;
+
+		if ( ExpressionInput->GetExpression() )
+		{
+			int32 ExpressionHash = ComputeMaterialExpressionHash( ExpressionInput->GetExpression() );
+			Hash = HashCombine( Hash, ExpressionHash );
+		}
+
+		Hash = HashCombine( Hash, GetTypeHash( ExpressionInput->GetOutputIndex() ) );
+
+		return Hash;
+	}
+
+	int32 ComputeMaterialExpressionHash( IDatasmithMaterialExpression* MaterialExpression )
+	{
+		uint32 Hash = GetTypeHash( MaterialExpression->GetType() );
+		Hash = HashCombine( Hash, GetTypeHash( FString( MaterialExpression->GetName() ) ) );
+
+		if ( MaterialExpression->IsA( EDatasmithMaterialExpressionType::TextureCoordinate ) )
+		{
+			IDatasmithMaterialExpressionTextureCoordinate* TextureCoordinate = static_cast< IDatasmithMaterialExpressionTextureCoordinate* >( MaterialExpression );
+			Hash = HashCombine( Hash, GetTypeHash( TextureCoordinate->GetCoordinateIndex() ) );
+			Hash = HashCombine( Hash, GetTypeHash( TextureCoordinate->GetUTiling() ) );
+			Hash = HashCombine( Hash, GetTypeHash( TextureCoordinate->GetVTiling() ) );
+		}
+		else if ( MaterialExpression->IsA( EDatasmithMaterialExpressionType::ConstantColor ) )
+		{
+			if ( FCString::Strlen( MaterialExpression->GetName() ) == 0 )
+			{
+				IDatasmithMaterialExpressionColor* ColorExpression = static_cast< IDatasmithMaterialExpressionColor* >( MaterialExpression );
+
+				Hash = HashCombine( Hash, GetTypeHash( ColorExpression->GetColor() ) );
+			}
+		}
+		else if ( MaterialExpression->IsA( EDatasmithMaterialExpressionType::ConstantScalar ) )
+		{
+			if ( FCString::Strlen( MaterialExpression->GetName() ) == 0 )
+			{
+				IDatasmithMaterialExpressionScalar* ScalarExpression = static_cast< IDatasmithMaterialExpressionScalar* >( MaterialExpression );
+
+				Hash = HashCombine( Hash, GetTypeHash( ScalarExpression->GetScalar() ) );
+			}
+		}
+		else if ( MaterialExpression->IsA( EDatasmithMaterialExpressionType::Generic ) )
+		{
+			IDatasmithMaterialExpressionGeneric* GenericExpression = static_cast< IDatasmithMaterialExpressionGeneric* >( MaterialExpression );
+
+			UClass* ExpressionClass = FindObject< UClass >( ANY_PACKAGE, *( FString( TEXT("MaterialExpression") ) + GenericExpression->GetExpressionName() ) );
+
+			UMaterialExpression* MaterialCDO = nullptr;
+
+			if ( ExpressionClass )
+			{
+				MaterialCDO = ExpressionClass->GetDefaultObject< UMaterialExpression >();
+			}
+
+			for ( int32 PropertyIndex = 0; PropertyIndex < GenericExpression->GetPropertiesCount(); ++PropertyIndex )
+			{
+				const TSharedPtr< IDatasmithKeyValueProperty >& KeyValue = GenericExpression->GetProperty( PropertyIndex );
+
+				if ( KeyValue )
+				{
+					Hash = HashCombine( Hash, GetTypeHash( KeyValue->GetName() ) );
+					Hash = HashCombine( Hash, GetTypeHash( KeyValue->GetPropertyType() ) );
+
+					// Only hash values if it's not the parameter
+					// Currently, if we're setting values on multiple properties, we're not sure which one is the parameter so we hash them all
+					if ( MaterialCDO && ( !MaterialCDO->HasAParameterName() || GenericExpression->GetPropertiesCount() > 1 ) )
+					{
+						Hash = HashCombine( Hash, GetTypeHash( KeyValue->GetValue() ) );
+					}
+				}
+			}
+		}
+		else if ( MaterialExpression->IsA( EDatasmithMaterialExpressionType::FunctionCall ) )
+		{
+			// Hash the path to the function as calling different functions should result in different hash values
+			IDatasmithMaterialExpressionFunctionCall* FunctionCallExpression = static_cast< IDatasmithMaterialExpressionFunctionCall* >( MaterialExpression );
+
+			Hash = HashCombine( Hash, GetTypeHash( FunctionCallExpression->GetFunctionPathName() ) );
+		}
+
+		for ( int32 InputIndex = 0; InputIndex < MaterialExpression->GetInputCount(); ++InputIndex )
+		{
+			Hash = HashCombine( Hash, ComputeExpressionInputHash( MaterialExpression->GetInput( InputIndex ) ) );
+		}
+
+		return Hash;
+	}
+
+	int32 ComputeMaterialHash( TSharedRef< IDatasmithUEPbrMaterialElement > MaterialElement )
+	{
+		int32 Hash = GetTypeHash( MaterialElement->GetTwoSided() );
+
+		Hash = HashCombine( Hash, GetTypeHash( MaterialElement->GetUseMaterialAttributes() ) );
+		Hash = HashCombine( Hash, GetTypeHash( MaterialElement->GetBlendMode() ) );
+
+		Hash = HashCombine( Hash, ComputeExpressionInputHash( &MaterialElement->GetBaseColor() ) );
+		Hash = HashCombine( Hash, ComputeExpressionInputHash( &MaterialElement->GetMetallic() ) );
+		Hash = HashCombine( Hash, ComputeExpressionInputHash( &MaterialElement->GetSpecular() ) );
+		Hash = HashCombine( Hash, ComputeExpressionInputHash( &MaterialElement->GetRoughness() ) );
+		Hash = HashCombine( Hash, ComputeExpressionInputHash( &MaterialElement->GetEmissiveColor() ) );
+		Hash = HashCombine( Hash, ComputeExpressionInputHash( &MaterialElement->GetOpacity() ) );
+		Hash = HashCombine( Hash, ComputeExpressionInputHash( &MaterialElement->GetNormal() ) );
+		Hash = HashCombine( Hash, ComputeExpressionInputHash( &MaterialElement->GetWorldDisplacement() ) );
+		Hash = HashCombine( Hash, ComputeExpressionInputHash( &MaterialElement->GetRefraction() ) );
+		Hash = HashCombine( Hash, ComputeExpressionInputHash( &MaterialElement->GetAmbientOcclusion() ) );
+		Hash = HashCombine( Hash, ComputeExpressionInputHash( &MaterialElement->GetMaterialAttributes() ) );
+
+		return Hash;
+	}
+}
+
+UMaterialFunction* FDatasmithMaterialImporter::CreateMaterialFunction( FDatasmithImportContext& ImportContext, const TSharedRef< IDatasmithBaseMaterialElement >& BaseMaterialElement )
+{
+	UMaterialFunction* MaterialFunction = nullptr;
+
+	if ( BaseMaterialElement->IsA( EDatasmithElementType::UEPbrMaterial ) )
+	{
+		const TSharedRef< IDatasmithUEPbrMaterialElement > MaterialElement = StaticCastSharedRef< IDatasmithUEPbrMaterialElement >( BaseMaterialElement );
+		UPackage* MaterialPackage = ImportContext.AssetsContext.MaterialFunctionsImportPackage.Get();
+		MaterialFunction = FDatasmithMaterialExpressions::CreateUEPbrMaterialFunction( MaterialPackage, MaterialElement, ImportContext.AssetsContext, nullptr, ImportContext.ObjectFlags );
+	}
+	else
+	{
+		//Only UEPbr materials should end up here.
+		check(false)
+		return nullptr;
+	}
+
+	if ( MaterialFunction != nullptr )
+	{
+		ImportContext.ImportedMaterialFunctions.Add( BaseMaterialElement ) = MaterialFunction;
+	}
+
+	return MaterialFunction;
+}
+
+UMaterialInterface* FDatasmithMaterialImporter::CreateMaterial( FDatasmithImportContext& ImportContext,
+	const TSharedRef< IDatasmithBaseMaterialElement >& BaseMaterialElement, UMaterialInterface* ExistingMaterial )
+{
+	UMaterialInterface* Material = nullptr;
+
+	if ( BaseMaterialElement->IsA( EDatasmithElementType::Material ) )
+	{
+		const TSharedRef< IDatasmithMaterialElement >& MaterialElement = StaticCastSharedRef< IDatasmithMaterialElement >( BaseMaterialElement );
+
+		UPackage* MaterialPackage = ImportContext.AssetsContext.MaterialsImportPackage.Get();
+
+		Material = FDatasmithMaterialExpressions::CreateDatasmithMaterial(MaterialPackage, MaterialElement, ImportContext.AssetsContext, nullptr, ImportContext.ObjectFlags);
+	}
+	else if ( BaseMaterialElement->IsA( EDatasmithElementType::MasterMaterial ) )
+	{
+		const TSharedRef< IDatasmithMasterMaterialElement >& MasterMaterialElement = StaticCastSharedRef< IDatasmithMasterMaterialElement >( BaseMaterialElement );
+		Material = ImportMasterMaterial( ImportContext, MasterMaterialElement, ExistingMaterial );
+	}
+	else if ( BaseMaterialElement->IsA( EDatasmithElementType::UEPbrMaterial ) )
+	{
+		const TSharedRef< IDatasmithUEPbrMaterialElement > MaterialElement = StaticCastSharedRef< IDatasmithUEPbrMaterialElement >( BaseMaterialElement );
+		if ( MaterialElement->GetMaterialFunctionOnly() )
+		{
+			//No need to instantiate a MaterialElement that is only used as a material function
+			return nullptr;
+		}
+
+		int32 MaterialHash = DatasmithMaterialImporterUtils::ComputeMaterialHash( MaterialElement );
+
+		if ( !ImportContext.ImportedParentMaterials.Contains( MaterialHash ) )
+		{
+			UMaterialInterface* ParentMaterial = FDatasmithMaterialExpressions::CreateUEPbrMaterial( ImportContext.AssetsContext.MasterMaterialsImportPackage.Get(), MaterialElement, ImportContext.AssetsContext, nullptr, ImportContext.ObjectFlags );
+
+			if (ParentMaterial == nullptr)
+			{
+				return nullptr;
+			}
+
+			ImportContext.ImportedParentMaterials.Add( MaterialHash ) = ParentMaterial;
+		}
+
+		// Always create a material instance
+		{
+			Material = FDatasmithMaterialExpressions::CreateUEPbrMaterialInstance( ImportContext.AssetsContext.MaterialsImportPackage.Get(), MaterialElement, ImportContext.AssetsContext,
+				Cast< UMaterial >( ImportContext.ImportedParentMaterials[ MaterialHash ] ), ImportContext.ObjectFlags );
+		}
+	}
+
+	if (Material != nullptr)
+	{
+		ImportContext.ImportedMaterials.Add( BaseMaterialElement ) = Material;
+	}
+
+	return Material;
+}
+
+UMaterialInterface* FDatasmithMaterialImporter::ImportMasterMaterial( FDatasmithImportContext& ImportContext, const TSharedRef< IDatasmithMasterMaterialElement >& MaterialElement, UMaterialInterface* ExistingMaterial )
+{
+	// Verify existing material is of the right class for further processing
+	UMaterialInstanceConstant* FoundConstantMaterial = Cast<UMaterialInstanceConstant>(ExistingMaterial);
+
+	FString Host = FDatasmithMasterMaterialManager::Get().GetHostFromString(ImportContext.Scene->GetHost());
+	TSharedPtr< FDatasmithMasterMaterialSelector > MaterialSelector = FDatasmithMasterMaterialManager::Get().GetSelector( *Host );
+
+	const FDatasmithMasterMaterial* MasterMaterial = nullptr;
+	FDatasmithMasterMaterial CustomMasterMaterial; // MasterMaterial might point on this so keep them in the same scope
+
+	if ( MaterialElement->GetMaterialType() == EDatasmithMasterMaterialType::Custom )
+	{
+		CustomMasterMaterial.FromSoftObjectPath( FSoftObjectPath( MaterialElement->GetCustomMaterialPathName() ) );
+
+		MasterMaterial = &CustomMasterMaterial;
+	}
+	else if ( MaterialSelector.IsValid() && MaterialSelector->IsValid() )
+	{
+		MasterMaterial = &MaterialSelector->GetMasterMaterial(MaterialElement);
+	}
+
+	if ( MasterMaterial && MasterMaterial->IsValid() )
+	{
+		const FString MaterialLabel = MaterialElement->GetLabel();
+		const FString MaterialName = MaterialLabel.Len() > 0 ? ImportContext.AssetsContext.MaterialNameProvider.GenerateUniqueName(MaterialLabel) : MaterialElement->GetName();
+
+		// Verify that the material could be created in final package
+		FText FailReason;
+		if (!FDatasmithImporterUtils::CanCreateAsset<UMaterialInstanceConstant>( ImportContext.AssetsContext.MaterialsFinalPackage.Get(), MaterialName, FailReason ))
+		{
+			ImportContext.LogError(FailReason);
+			return nullptr;
+		}
+
+		UMaterialInstanceConstant* MaterialInstance = FoundConstantMaterial;
+
+		if (MaterialInstance == nullptr)
+		{
+			MaterialInstance = NewObject<UMaterialInstanceConstant>(ImportContext.AssetsContext.MaterialsImportPackage.Get(), *MaterialName, ImportContext.ObjectFlags);
+			MaterialInstance->Parent = MasterMaterial->GetMaterial();
+
+			FAssetRegistryModule::AssetCreated(MaterialInstance);
+		}
+		else
+		{
+			MaterialInstance = DuplicateObject< UMaterialInstanceConstant >(MaterialInstance, ImportContext.AssetsContext.MaterialsImportPackage.Get(), *MaterialName);
+			IDatasmithImporterModule::Get().ResetOverrides(MaterialInstance); // Don't copy the existing overrides
+		}
+
+		UDatasmithMaterialInstanceTemplate* MaterialInstanceTemplate = NewObject< UDatasmithMaterialInstanceTemplate >( MaterialInstance );
+
+		// Find matching master material parameters
+		for (int i = 0; i < MaterialElement->GetPropertiesCount(); ++i)
+		{
+			const TSharedPtr< IDatasmithKeyValueProperty > Property = MaterialElement->GetProperty(i);
+			FString PropertyName(Property->GetName());
+
+			// Vector Params
+			if ( MasterMaterial->VectorParams.Contains(PropertyName) )
+			{
+				FLinearColor Color;
+				if ( MaterialSelector->GetColor( Property, Color ) )
+				{
+					MaterialInstanceTemplate->VectorParameterValues.Add( FName(*PropertyName), Color );
+				}
+			}
+			// Scalar Params
+			else if ( MasterMaterial->ScalarParams.Contains(PropertyName) )
+			{
+				float Value;
+				if ( MaterialSelector->GetFloat( Property, Value ) )
+				{
+					MaterialInstanceTemplate->ScalarParameterValues.Add( FName(*PropertyName), Value );
+				}
+			}
+			// Bool Params
+			else if (MasterMaterial->BoolParams.Contains(PropertyName))
+			{
+				bool bValue;
+				if ( MaterialSelector->GetBool( Property, bValue ) )
+				{
+					MaterialInstanceTemplate->StaticParameters.StaticSwitchParameters.Add( FName( Property->GetName() ), bValue );
+				}
+			}
+			// Texture Params
+			else if (MasterMaterial->TextureParams.Contains(PropertyName))
+			{
+				FString TexturePath;
+				if ( MaterialSelector->GetTexture( Property, TexturePath ) )
+				{
+					FString TextureName = ObjectTools::SanitizeObjectName(FPaths::GetBaseFilename( TexturePath ));
+
+					UTexture* Texture = FDatasmithImporterUtils::FindAsset< UTexture >( ImportContext.AssetsContext, *TextureName );
+					MaterialInstanceTemplate->TextureParameterValues.Add( FName(*PropertyName), Texture );
+
+					//If we are adding a virtual texture to a non-virtual texture streamer then we will need to convert back that Virtual texture.
+					UTexture* DefaultTextureValue = nullptr;
+					UTexture2D* Texture2D = Cast<UTexture2D>(Texture);
+					if (Texture2D && Texture2D->VirtualTextureStreaming &&
+						MaterialInstance->GetTextureParameterDefaultValue(FName(*PropertyName), DefaultTextureValue) && DefaultTextureValue)
+					{
+						if (!DefaultTextureValue->VirtualTextureStreaming)
+						{
+							ImportContext.AssetsContext.VirtualTexturesToConvert.Add(Texture2D);
+						}
+					}
+				}
+			}
+		}
+
+		MaterialInstanceTemplate->Apply( MaterialInstance );
+
+		if (MaterialSelector.IsValid() && MaterialSelector->IsValid())
+		{
+			MaterialSelector->FinalizeMaterialInstance(MaterialElement, MaterialInstance);
+		}
+
+		return MaterialInstance;
+	}
+
+	return nullptr;
+}
+
+int32 FDatasmithMaterialImporter::GetMaterialRequirements(UMaterialInterface * MaterialInterface)
+{
+	if (MaterialInterface == nullptr || MaterialInterface->GetMaterial() == nullptr)
+	{
+		return EMaterialRequirements::RequiresNothing;
+	}
+	// Currently all Datasmith materials require at least normals and tangents
+	// @todo: Adjust initial value and logic based on future materials' requirements
+	int32 MaterialRequirement = EMaterialRequirements::RequiresNormals | EMaterialRequirements::RequiresTangents;
+
+	UMaterial* Material = MaterialInterface->GetMaterial();
+	// Material with displacement or support for PNT requires adjacency and has their TessellationMultiplier set
+	if (Material->TessellationMultiplier.Expression != nullptr || Material->D3D11TessellationMode != EMaterialTessellationMode::MTM_NoTessellation)
+	{
+		MaterialRequirement |= EMaterialRequirements::RequiresAdjacency;
+	}
+
+	return MaterialRequirement;
+}
diff --git a/Engine/Plugins/Enterprise/DatasmithImporter/Source/DatasmithImporter/Private/DatasmithMeshHelper.cpp b/Engine/Plugins/Enterprise/DatasmithImporter/Source/DatasmithImporter/Private/DatasmithMeshHelper.cpp
new file mode 100644
--- /dev/null
+++ b/Engine/Plugins/Enterprise/DatasmithImporter/Source/DatasmithImporter/Private/DatasmithMeshHelper.cpp
@@ -0,0 +1,357 @@
+// Copyright 1998-2019 Epic Games, Inc. All Rights Reserved.
+
+#include "DatasmithMeshHelper.h"
+
+#include "StaticMeshAttributes.h"
+#include "Engine/StaticMesh.h"
+#include "MeshDescriptionOperations.h"
+#include "DatasmithMesh.h"
+#include "RawMesh.h"
+#include "DatasmithUtils.h"
+#include "DatasmithMeshUObject.h"
+#include "Misc/SecureHash.h"
+#include "UVMapSettings.h"
+
+
+namespace DatasmithMeshHelper
+{
+
+	int32 GetPolygonCount(const FMeshDescription& Mesh)
+	{
+		return Mesh.Polygons().Num();
+	}
+
+	int32 GetTriangleCount(const FMeshDescription& Mesh)
+	{
+		return Mesh.Triangles().Num();
+	}
+
+	void ExtractVertexPositions(const FMeshDescription& Mesh, TArray<FVector>& Positions)
+	{
+		FStaticMeshConstAttributes Attributes(Mesh);
+		const TVertexAttributesConstRef<FVector> VertexPositions = Attributes.GetVertexPositions();
+		if (VertexPositions.IsValid())
+		{
+			Positions.Reserve(VertexPositions.GetNumElements());
+			for (const FVertexID VertexID : Mesh.Vertices().GetElementIDs())
+			{
+				Positions.Add(VertexPositions[VertexID]);
+			}
+		}
+	}
+
+	void PrepareAttributeForStaticMesh(FMeshDescription& Mesh)
+	{
+		FStaticMeshAttributes(Mesh).Register();
+	}
+
+	bool IsTriangleDegenerated(const FMeshDescription& Mesh, const FMeshTriangle& MeshTriangle)
+	{
+		// extract triangle vertices positions
+		FVector P[3];
+		const TVertexAttributesConstRef<FVector> VertexPositions = Mesh.VertexAttributes().GetAttributesRef<FVector>(MeshAttribute::Vertex::Position);
+		for (int32 CornerIndex = 0; CornerIndex < 3; ++CornerIndex)
+		{
+			FVertexID VertexID = Mesh.GetVertexInstanceVertex(MeshTriangle.GetVertexInstanceID(CornerIndex));
+			P[CornerIndex] = VertexPositions[VertexID];
+		}
+
+		float NormalLengthSquare = ((P[1] - P[2]) ^ (P[0] - P[2])).SizeSquared();
+		return NormalLengthSquare < SMALL_NUMBER;
+	}
+
+	FMeshDescription* InitMeshDescription(UStaticMesh* StaticMesh, int32 LodIndex)
+	{
+		if (!ensure(StaticMesh))
+		{
+			return nullptr;
+		}
+
+		while (!StaticMesh->IsSourceModelValid(LodIndex))
+		{
+			StaticMesh->AddSourceModel();
+		}
+
+		return StaticMesh->CreateMeshDescription(LodIndex);
+	}
+
+	void FillUStaticMesh(UStaticMesh* StaticMesh, int32 LodIndex, FRawMesh& RawMesh, const TMap<int32, FName>* InMaterialMapInverse)
+	{
+		FMeshDescription MeshDescription;
+		PrepareAttributeForStaticMesh(MeshDescription);
+		if (InMaterialMapInverse == nullptr)
+		{
+			// no explicit mapping, build one from UStaticMesh
+			TMap<FName, int32> MaterialMap; // unused
+			TMap<int32, FName> MaterialMapInverse;
+
+ 			BuildMaterialMappingFromStaticMesh(StaticMesh, MaterialMap, MaterialMapInverse);
+			FMeshDescriptionOperations::ConvertFromRawMesh(RawMesh, MeshDescription, MaterialMapInverse);
+		}
+		else
+		{
+			FMeshDescriptionOperations::ConvertFromRawMesh(RawMesh, MeshDescription, *InMaterialMapInverse);
+		}
+
+		FillUStaticMesh(StaticMesh, LodIndex, MoveTemp(MeshDescription));
+	}
+
+	void FillUStaticMesh(UStaticMesh* StaticMesh, int32 LodIndex, FDatasmithMesh& DsMesh)
+	{
+		FMeshDescription MeshDescription;
+		PrepareAttributeForStaticMesh(MeshDescription);
+		if (FDatasmithMeshUtils::ToMeshDescription(DsMesh, MeshDescription))
+		{
+			FillUStaticMesh(StaticMesh, LodIndex, MoveTemp(MeshDescription));
+		}
+	}
+
+	void FillUStaticMesh(UStaticMesh* StaticMesh, int LodIndex, FMeshDescription&& MeshDescription)
+	{
+		if ( !IsMeshValid( MeshDescription ) )
+		{
+			return;
+		}
+
+		FMeshDescription* OriginalMeshDescription = InitMeshDescription(StaticMesh, LodIndex);
+		if (ensure(OriginalMeshDescription))
+		{
+			// #ueent_todo: validate attributes too make sure we don't push incompatible MeshDescriptions
+			*OriginalMeshDescription = MoveTemp(MeshDescription);
+		}
+	}
+
+	void FillUStaticMeshByCopy(UStaticMesh* StaticMesh, int LodIndex, FMeshDescription MeshDescription)
+	{
+		FillUStaticMesh(StaticMesh, LodIndex, MoveTemp(MeshDescription));
+	}
+
+	void PrepareStaticMaterials(UStaticMesh* StaticMesh, int32 MaterialCount)
+	{
+		if (!ensure(StaticMesh))
+		{
+			return;
+		}
+
+		MaterialCount = FMath::Max(MaterialCount, 1);
+
+		// Don't touch the StaticMesh if it's not needed (would break the template system on reimport)
+		if (MaterialCount != StaticMesh->StaticMaterials.Num())
+		{
+			StaticMesh->StaticMaterials.Empty(MaterialCount);
+			// Declare materials applied to the StaticMesh according to fit the logic in UDatasmithImportFactory::UpdateMaterials
+			for (int MaterialIndex = 0; MaterialIndex < MaterialCount; MaterialIndex++)
+			{
+				StaticMesh->GetSectionInfoMap().Set(0, MaterialIndex, FMeshSectionInfo(MaterialIndex));
+				FName SlotName = DefaultSlotName(MaterialIndex);
+				StaticMesh->StaticMaterials.Add(FStaticMaterial(nullptr, SlotName, SlotName));
+			}
+		}
+	}
+
+	void BuildMaterialMappingFromStaticMesh(UStaticMesh* StaticMesh, TMap<FName, int32>& MaterialMap, TMap<int32, FName>& MaterialMapInverse)
+	{
+		if (!ensure(StaticMesh))
+		{
+			return;
+		}
+
+		int32 MatCount = StaticMesh->StaticMaterials.Num();
+		ensure(MatCount);
+		MaterialMap.Empty(MatCount);
+		MaterialMapInverse.Empty(MatCount);
+
+		for ( int32 MaterialIndex = 0; MaterialIndex < MatCount; ++MaterialIndex )
+		{
+			FName& MatName = StaticMesh->StaticMaterials[MaterialIndex].ImportedMaterialSlotName;
+			MaterialMap.Add(MatName, MaterialIndex);
+			MaterialMapInverse.Add(MaterialIndex, MatName);
+		}
+	}
+
+	void HashMeshLOD(UStaticMesh* StaticMesh, int32 LodIndex, FMD5& MD5)
+	{
+		if (ensure(StaticMesh))
+		{
+			const FMeshDescription* Mesh = StaticMesh->GetMeshDescription(LodIndex);
+			if (ensure(Mesh))
+			{
+				HashMeshDescription(*Mesh, MD5);
+			}
+		}
+	}
+
+	void HashMeshDescription(const FMeshDescription& Mesh, FMD5& MD5)
+	{
+		auto HashAttributeSet = [](FMD5& MD5, const FAttributesSetBase& AttributeSet)
+		{
+			TArray<FName> OutAttributeNames;
+			AttributeSet.GetAttributeNames(OutAttributeNames);
+			for (const FName& AttributeName: OutAttributeNames)
+			{
+				int32 AttributeHash = AttributeSet.GetHash(AttributeName);
+				MD5.Update((uint8*)&AttributeHash, sizeof(AttributeHash));
+				// #ueent_todo better hash function...
+				// see FDatasmithFBXSceneMesh::GetHash
+			}
+		};
+
+		HashAttributeSet(MD5, Mesh.VertexAttributes());
+		HashAttributeSet(MD5, Mesh.VertexInstanceAttributes());
+		HashAttributeSet(MD5, Mesh.EdgeAttributes());
+		HashAttributeSet(MD5, Mesh.PolygonAttributes());
+		HashAttributeSet(MD5, Mesh.PolygonGroupAttributes());
+	}
+
+	FName DefaultSlotName(int32 MaterialIndex)
+	{
+		return FName( *FString::FromInt(MaterialIndex) );
+	}
+
+	int32 GetNumUVChannel(const FMeshDescription& Mesh)
+	{
+		const TVertexInstanceAttributesConstRef<FVector2D> UVChannels = Mesh.VertexInstanceAttributes().GetAttributesRef<FVector2D>(MeshAttribute::VertexInstance::TextureCoordinate);
+		return UVChannels.GetNumIndices();
+	}
+
+	bool HasUVChannel(const FMeshDescription& Mesh, int32 ChannelIndex)
+	{
+		return ChannelIndex >= 0 && ChannelIndex < GetNumUVChannel(Mesh);
+	}
+
+	bool HasUVData(const FMeshDescription& Mesh, int32 ChannelIndex)
+	{
+		if (HasUVChannel(Mesh, ChannelIndex))
+		{
+			const TVertexInstanceAttributesConstRef<FVector2D> UVChannels = Mesh.VertexInstanceAttributes().GetAttributesRef<FVector2D>(MeshAttribute::VertexInstance::TextureCoordinate);
+			const FVector2D DefValue = UVChannels.GetDefaultValue();
+			for (FVertexInstanceID InstanceID : Mesh.VertexInstances().GetElementIDs())
+			{
+				if (UVChannels.Get(InstanceID, ChannelIndex) != DefValue)
+				{
+					return true;
+				}
+			}
+		}
+		return false;
+	}
+
+	bool RequireUVChannel(FMeshDescription& Mesh, int32 ChannelIndex)
+	{
+		if (ChannelIndex < 0 || ChannelIndex >= MAX_MESH_TEXTURE_COORDS_MD)
+		{
+			return false;
+		}
+
+		if (!HasUVChannel(Mesh, ChannelIndex))
+		{
+			TVertexInstanceAttributesRef<FVector2D> UVChannels = Mesh.VertexInstanceAttributes().GetAttributesRef<FVector2D>(MeshAttribute::VertexInstance::TextureCoordinate);
+			UVChannels.SetNumIndices(ChannelIndex + 1);
+		}
+		return true;
+	}
+
+	void CreateDefaultUVs(FDatasmithMesh& DatasmithMesh)
+	{
+		// Get the mesh description to generate BoxUV.
+		FMeshDescription MeshDescription;
+		PrepareAttributeForStaticMesh(MeshDescription);
+		FDatasmithMeshUtils::ToMeshDescription(DatasmithMesh, MeshDescription);
+		FUVMapParameters UVParameters(DatasmithMesh.GetExtents().GetCenter(), FQuat::Identity, DatasmithMesh.GetExtents().GetSize(), FVector::OneVector, FVector2D::UnitVector);
+		TMap<FVertexInstanceID, FVector2D> TexCoords;
+		FMeshDescriptionOperations::GenerateBoxUV(MeshDescription, UVParameters, TexCoords);
+
+		// Put the results in a map to determine the number of unique values.
+		TMap<FVector2D, TArray<int32>> UniqueTexCoordMap;
+		for (const TPair<FVertexInstanceID, FVector2D>& Pair : TexCoords)
+		{
+			TArray<int32>& MappedIndices = UniqueTexCoordMap.FindOrAdd(Pair.Value);
+			MappedIndices.Add(Pair.Key.GetValue());
+		}
+
+		//Set the UV values
+		if (DatasmithMesh.GetUVChannelsCount() == 0)
+		{
+			DatasmithMesh.AddUVChannel();
+		}
+		DatasmithMesh.SetUVCount(0, UniqueTexCoordMap.Num());
+		int32 UVIndex = 0;
+		TArray<int32> IndicesMapping;
+		IndicesMapping.AddZeroed(TexCoords.Num());
+
+		for (const TPair<FVector2D, TArray<int32>>& UniqueCoordPair : UniqueTexCoordMap)
+		{
+			DatasmithMesh.SetUV(0, UVIndex, UniqueCoordPair.Key.X, UniqueCoordPair.Key.Y);
+			for (int32 IndicesIndex : UniqueCoordPair.Value)
+			{
+				IndicesMapping[IndicesIndex] = UVIndex;
+			}
+			UVIndex++;
+		}
+
+		//Map the UV indices.
+		for (int32 FaceIndex = 0; FaceIndex < DatasmithMesh.GetFacesCount(); ++FaceIndex)
+		{
+			const int32 IndicesOffset = FaceIndex * 3;
+			check(IndicesOffset + 2 < IndicesMapping.Num());
+			DatasmithMesh.SetFaceUV(FaceIndex, 0, IndicesMapping[IndicesOffset + 0], IndicesMapping[IndicesOffset + 1], IndicesMapping[IndicesOffset + 2]);
+		}
+	}
+
+	void CreateDefaultUVs(FMeshDescription& MeshDescription)
+	{
+		FBox MeshBoundingBox = MeshDescription.ComputeBoundingBox();
+		FUVMapParameters UVParameters(MeshBoundingBox.GetCenter(), FQuat::Identity, MeshBoundingBox.GetExtent(), FVector::OneVector, FVector2D::UnitVector);
+		TMap<FVertexInstanceID, FVector2D> TexCoords;
+		FMeshDescriptionOperations::GenerateBoxUV(MeshDescription, UVParameters, TexCoords);
+
+		TMeshAttributesRef<FVertexInstanceID, FVector2D> UVs = MeshDescription.VertexInstanceAttributes().GetAttributesRef<FVector2D>(MeshAttribute::VertexInstance::TextureCoordinate);
+		if (UVs.GetNumIndices() == 0)
+		{
+			UVs.SetNumIndices(1);
+		}
+
+		for (const FVertexInstanceID& VertexInstanceID : MeshDescription.VertexInstances().GetElementIDs())
+		{
+			if (const FVector2D* UVCoord = TexCoords.Find(VertexInstanceID))
+			{
+				UVs.Set(VertexInstanceID, 0, *UVCoord);
+			}
+			else
+			{
+				ensureMsgf(false, TEXT("Tried to apply UV data that did not match the MeshDescription."));
+			}
+		}
+	}
+
+	bool IsMeshValid(const FMeshDescription& Mesh, FVector BuildScale)
+	{
+		const TVertexAttributesConstRef<FVector> VertexPositions = Mesh.VertexAttributes().GetAttributesRef<FVector>(MeshAttribute::Vertex::Position);
+		FVector RawNormalScale(BuildScale.Y*BuildScale.Z, BuildScale.X*BuildScale.Z, BuildScale.X*BuildScale.Y); // Component-wise scale
+		for (const FPolygonID PolygonID : Mesh.Polygons().GetElementIDs())
+		{
+			FVector Corners[3];
+			for (const FTriangleID TriangleID : Mesh.GetPolygonTriangleIDs(PolygonID))
+			{
+				for (int32 CornerIndex = 0; CornerIndex < 3; ++CornerIndex)
+				{
+					FVertexInstanceID VertexInstanceID = Mesh.GetTriangleVertexInstance(TriangleID, CornerIndex);
+					FVertexID VertexID = Mesh.GetVertexInstanceVertex(VertexInstanceID);
+					Corners[CornerIndex] = VertexPositions[VertexID];
+				}
+
+				FVector RawNormal = (Corners[1] - Corners[2]) ^ (Corners[0] - Corners[2]);
+				RawNormal *= RawNormalScale;
+				double FourSquaredTriangleArea = RawNormal.SizeSquared();
+				if (FourSquaredTriangleArea > SMALL_NUMBER)
+				{
+					return true;
+				}
+			}
+		}
+
+		// all faces are degenerated, mesh is invalid
+		return false;
+	}
+
+} // namespace DatasmithMeshHelper
diff --git a/Engine/Plugins/Enterprise/DatasmithImporter/Source/DatasmithImporter/Private/DatasmithSceneSource.cpp b/Engine/Plugins/Enterprise/DatasmithImporter/Source/DatasmithImporter/Private/DatasmithSceneSource.cpp
new file mode 100644
--- /dev/null
+++ b/Engine/Plugins/Enterprise/DatasmithImporter/Source/DatasmithImporter/Private/DatasmithSceneSource.cpp
@@ -0,0 +1,29 @@
+// Copyright 1998-2019 Epic Games, Inc. All Rights Reserved.
+
+#include "DatasmithSceneSource.h"
+
+#include "Translators/DatasmithTranslator.h"
+#include "Translators/DatasmithTranslatorManager.h"
+
+#include "Misc/Paths.h"
+
+#define LOCTEXT_NAMESPACE "DatasmithSceneSource"
+
+void FDatasmithSceneSource::SetSourceFile(const FString& InFilePath)
+{
+	FilePath = InFilePath;
+	SceneDeducedName = FPaths::GetBaseFilename(FilePath);
+	FileExtension = FPaths::GetExtension(FilePath);
+}
+
+void FDatasmithSceneSource::SetSceneName(const FString& InSceneName)
+{
+	SceneOverrideName = InSceneName;
+}
+
+const FString& FDatasmithSceneSource::GetSceneName() const
+{
+	return SceneOverrideName.IsEmpty() ? SceneDeducedName : SceneOverrideName;
+}
+
+#undef LOCTEXT_NAMESPACE
diff --git a/Engine/Plugins/Enterprise/DatasmithImporter/Source/DatasmithImporter/Private/DatasmithStaticMeshImporter.cpp b/Engine/Plugins/Enterprise/DatasmithImporter/Source/DatasmithImporter/Private/DatasmithStaticMeshImporter.cpp
new file mode 100644
--- /dev/null
+++ b/Engine/Plugins/Enterprise/DatasmithImporter/Source/DatasmithImporter/Private/DatasmithStaticMeshImporter.cpp
@@ -0,0 +1,693 @@
+// Copyright 1998-2019 Epic Games, Inc. All Rights Reserved.
+
+#include "DatasmithStaticMeshImporter.h"
+
+#include "DatasmithImportContext.h"
+#include "DatasmithImporterModule.h"
+#include "DatasmithMaterialImporter.h"
+#include "DatasmithMeshHelper.h"
+#include "DatasmithMeshUObject.h"
+#include "IDatasmithSceneElements.h"
+#include "ObjectTemplates/DatasmithStaticMeshTemplate.h"
+#include "Translators/DatasmithPayload.h"
+#include "Utility/DatasmithImporterUtils.h"
+
+#include "Algo/AnyOf.h"
+#include "Async/Async.h"
+#include "Engine/StaticMesh.h"
+#include "LayoutUV.h"
+#include "MeshBuild.h"
+#include "MeshDescriptionOperations.h"
+#include "MeshUtilities.h"
+#include "Misc/FeedbackContext.h"
+#include "Misc/Paths.h"
+#include "Misc/ScopedSlowTask.h"
+#include "OverlappingCorners.h"
+#include "PhysicsEngine/AggregateGeom.h"
+#include "PhysicsEngine/BodySetup.h"
+#include "StaticMeshAttributes.h"
+#include "Templates/UniquePtr.h"
+#include "UObject/Package.h"
+
+#define LOCTEXT_NAMESPACE "DatasmithImporter"
+
+namespace DatasmithStaticMeshImporterImpl
+{
+	float Get2DSurface(const FVector4& Dimensions)
+	{
+		if (Dimensions[0] >= Dimensions[1] && Dimensions[2] >= Dimensions[1])
+		{
+			return Dimensions[0]*Dimensions[2];
+		}
+		if (Dimensions[0] >= Dimensions[2] && Dimensions[1] >= Dimensions[2])
+		{
+			return Dimensions[0]*Dimensions[1];
+		}
+		return Dimensions[2]*Dimensions[1];
+	}
+
+	float CalcBlendWeight(const FVector4& Dimensions, float MaxArea, float Max2DSurface)
+	{
+		float Current2DSurface = Get2DSurface(Dimensions);
+		float Weight = FMath::Sqrt((Dimensions[3] / MaxArea)) + FMath::Sqrt(Current2DSurface / Max2DSurface);
+		return Weight;
+	}
+
+	int32 GetLightmapSize(float Weight, int32 MinimumSize, int32 MaximumSize)
+	{
+		Weight = FMath::Clamp(Weight, 0.f, 1.f);
+
+		float WeightedSize = (float(MinimumSize)*(1.0f - Weight)) + Weight*float(MaximumSize);
+
+		int32 LightmapSizeValue = MinimumSize;
+		while (WeightedSize > LightmapSizeValue)
+		{
+			LightmapSizeValue *= 2;
+		}
+
+		return FMath::Clamp( LightmapSizeValue, 0, 4096 );
+	}
+
+	TMap< int32, FString > GetMaterials( TSharedRef< IDatasmithMeshElement > MeshElement )
+	{
+		TMap< int32, FString > Materials;
+		for ( int i = 0; i < MeshElement->GetMaterialSlotCount(); ++i )
+		{
+			Materials.Add( MeshElement->GetMaterialSlotAt(i)->GetId(), MeshElement->GetMaterialSlotAt(i)->GetName() );
+		}
+		return Materials;
+	}
+
+	int32 GetStaticMeshBuildRequirements( const FDatasmithAssetsImportContext& AssetsContext, TSharedRef< IDatasmithMeshElement > MeshElement )
+	{
+		int32 BuildRequirements = EMaterialRequirements::RequiresNormals | EMaterialRequirements::RequiresTangents;;
+		for ( const TPair<int32, FString>& Material : GetMaterials( MeshElement ) )
+		{
+			if ( AssetsContext.MaterialsRequirements.Contains( *Material.Value ) )
+			{
+				BuildRequirements |= AssetsContext.MaterialsRequirements[ *Material.Value ];
+			}
+		}
+		return BuildRequirements;
+	}
+}
+
+UStaticMesh* FDatasmithStaticMeshImporter::ImportStaticMesh(const TSharedRef< IDatasmithMeshElement > MeshElement, FDatasmithMeshElementPayload& Payload, EObjectFlags ObjectFlags, const FDatasmithStaticMeshImportOptions& ImportOptions, FDatasmithAssetsImportContext& AssetsContext, UStaticMesh* ExistingMesh)
+{
+	TRACE_CPUPROFILER_EVENT_SCOPE(FDatasmithStaticMeshImporter::ImportStaticMesh);
+
+	// 1. import data
+	TArray< FMeshDescription >& MeshDescriptions = Payload.LodMeshes;
+
+	// Update MeshElement's LOD count with the actual number of mesh descriptions
+	MeshElement->SetLODCount( MeshDescriptions.Num() );
+	if ( MeshDescriptions.Num() == 0 )
+	{
+		return nullptr;
+	}
+
+	// 2. find the destination
+	UStaticMesh* ResultStaticMesh = nullptr;
+	FString StaticMeshName = AssetsContext.StaticMeshNameProvider.GenerateUniqueName( MeshElement->GetLabel() );
+	UPackage* Outer = AssetsContext.StaticMeshesImportPackage.Get();
+
+	// Verify that the static mesh could be created in final package
+	FText FailReason;
+	if (!FDatasmithImporterUtils::CanCreateAsset<UStaticMesh>( AssetsContext.StaticMeshesFinalPackage.Get(), StaticMeshName, FailReason ))
+	{
+		AssetsContext.ParentContext.LogError(FailReason);
+		return nullptr;
+	}
+
+	if ( ExistingMesh )
+	{
+		if ( ExistingMesh->GetOuter() != Outer )
+		{
+			ResultStaticMesh = DuplicateObject< UStaticMesh >( ExistingMesh, Outer, *StaticMeshName );
+			IDatasmithImporterModule::Get().ResetOverrides( ResultStaticMesh ); // Don't copy the existing overrides
+		}
+		else
+		{
+			ResultStaticMesh = ExistingMesh;
+		}
+
+		ResultStaticMesh->SetFlags( ObjectFlags );
+	}
+	else
+	{
+		ResultStaticMesh = NewObject< UStaticMesh >( Outer, *StaticMeshName, ObjectFlags );
+	}
+
+	// 3. Write data to destination
+	int32 LODIndex = 0;
+	for (FMeshDescription& MeshDesc : MeshDescriptions)
+	{
+		DatasmithMeshHelper::FillUStaticMesh(ResultStaticMesh, LODIndex, MoveTemp(MeshDesc));
+		LODIndex++;
+	}
+
+	// 4. Collisions
+	TArray< FVector > VertexPositions;
+	DatasmithMeshHelper::ExtractVertexPositions(Payload.CollisionMesh, VertexPositions);
+	if ( VertexPositions.Num() == 0 )
+	{
+		VertexPositions = MoveTemp(Payload.CollisionPointCloud);
+	}
+
+	if ( VertexPositions.Num() > 0 )
+	{
+		ProcessCollision( ResultStaticMesh, VertexPositions );
+	}
+
+	return ResultStaticMesh;
+}
+
+bool FDatasmithStaticMeshImporter::ShouldRecomputeNormals(const FMeshDescription& MeshDescription, int32 BuildRequirements)
+{
+	const TVertexInstanceAttributesConstRef<FVector> Normals = MeshDescription.VertexInstanceAttributes().GetAttributesRef<FVector>(MeshAttribute::VertexInstance::Normal);
+	check(Normals.IsValid());
+	return Algo::AnyOf(MeshDescription.VertexInstances().GetElementIDs(), [&](const FVertexInstanceID& InstanceID) { return !Normals[InstanceID].IsNormalized(); });
+}
+
+bool FDatasmithStaticMeshImporter::ShouldRecomputeTangents(const FMeshDescription& MeshDescription, int32 BuildRequirements)
+{
+	const TVertexInstanceAttributesConstRef<FVector> Tangents = MeshDescription.VertexInstanceAttributes().GetAttributesRef<FVector>(MeshAttribute::VertexInstance::Tangent);
+	check(Tangents.IsValid());
+	return Algo::AnyOf(MeshDescription.VertexInstances().GetElementIDs(), [&](const FVertexInstanceID& InstanceID) { return !Tangents[InstanceID].IsNormalized(); });
+}
+
+bool FDatasmithStaticMeshImporter::PreBuildStaticMesh( UStaticMesh* StaticMesh )
+{
+	TRACE_CPUPROFILER_EVENT_SCOPE(FDatasmithStaticMeshImporter::PreBuildStaticMesh);
+
+	// The input static mesh must at least have its LOD 0 valid
+	if (!StaticMesh->IsMeshDescriptionValid(0))
+	{
+		return false;
+	}
+
+	for (int32 LodIndex = 0; LodIndex < StaticMesh->GetNumSourceModels(); ++LodIndex)
+	{
+		if (!StaticMesh->IsMeshDescriptionValid(LodIndex))
+		{
+			continue;
+		}
+
+		FMeshBuildSettings& BuildSettings = StaticMesh->GetSourceModel(LodIndex).BuildSettings;
+		FMeshDescription& MeshDescription = *StaticMesh->GetMeshDescription(LodIndex);
+
+		// We should always have some UV data in channel 0 because it is used in the mesh tangent calculation during the build.
+		if (!DatasmithMeshHelper::HasUVData(MeshDescription, 0))
+		{
+			DatasmithMeshHelper::CreateDefaultUVs(MeshDescription);
+		}
+
+		// If lightmap UVs are not to be generated, set destination lightmap to 0.
+		if (!BuildSettings.bGenerateLightmapUVs)
+		{
+			BuildSettings.DstLightmapIndex = 0;
+		}
+		else if (!DatasmithMeshHelper::HasUVData(MeshDescription, BuildSettings.SrcLightmapIndex))
+		{
+			//We know that a UV at channel 0 exists so we fallback on it.
+			BuildSettings.SrcLightmapIndex = 0;
+		}
+
+		if (DatasmithMeshHelper::IsMeshValid(MeshDescription, BuildSettings.BuildScale3D))
+		{
+			DatasmithMeshHelper::RequireUVChannel(MeshDescription, BuildSettings.DstLightmapIndex);
+			UStaticMesh::FCommitMeshDescriptionParams Params;
+
+			// We will call MarkPackageDirty() ourself from the main thread
+			Params.bMarkPackageDirty = false;
+
+			// Ensure we get DDC benefits when we import already cached content
+			Params.bUseHashAsGuid = true;
+
+			StaticMesh->CommitMeshDescription(LodIndex, Params);
+		}
+		else
+		{
+			return false;
+		}
+	}
+
+	return true;
+}
+
+void FDatasmithStaticMeshImporter::BuildStaticMesh( UStaticMesh* StaticMesh )
+{
+	BuildStaticMeshes({ StaticMesh });
+}
+
+void FDatasmithStaticMeshImporter::PostMeshBuild(UStaticMesh* StaticMesh)
+{
+	StaticMesh->ClearMeshDescriptions();
+}
+
+void FDatasmithStaticMeshImporter::BuildStaticMeshes(const TArray< UStaticMesh* >& StaticMeshes, TFunction<bool(UStaticMesh*)> ProgressCallback)
+{
+	TRACE_CPUPROFILER_EVENT_SCOPE(FDatasmithStaticMeshImporter::BuildStaticMeshes);
+
+	TArray<FDelegateHandle> PostMeshBuildDelegateHandles;
+	PostMeshBuildDelegateHandles.SetNum(StaticMeshes.Num());
+
+	for (int32 StaticMeshIndex = 0; StaticMeshIndex < StaticMeshes.Num(); ++StaticMeshIndex)
+	{
+		// Release mesh descriptions as soon as they've been serialized and built to reduce memory footprint during import
+		PostMeshBuildDelegateHandles[StaticMeshIndex] = StaticMeshes[StaticMeshIndex]->OnPostMeshBuild().AddStatic(&FDatasmithStaticMeshImporter::PostMeshBuild);
+	}
+
+	UStaticMesh::BatchBuild(StaticMeshes, true, ProgressCallback);
+
+	for (int32 StaticMeshIndex = 0; StaticMeshIndex < StaticMeshes.Num(); ++StaticMeshIndex)
+	{
+		StaticMeshes[StaticMeshIndex]->OnPostMeshBuild().Remove(PostMeshBuildDelegateHandles[StaticMeshIndex]);
+	}
+}
+
+void FDatasmithStaticMeshImporter::ProcessCollision(UStaticMesh* StaticMesh, const TArray< FVector >& VertexPositions)
+{
+	// The following code is copied from StaticMeshEdit AddConvexGeomFromVertices (inaccessible outside UnrealEd)
+	if ( !StaticMesh )
+	{
+		return;
+	}
+
+	StaticMesh->bCustomizedCollision = true;
+	StaticMesh->CreateBodySetup();
+	if ( !ensure(StaticMesh->BodySetup) )
+	{
+		return;
+	}
+
+	// Convex elements must be removed first since the re-import process uses the same flow
+	FKAggregateGeom& AggGeom = StaticMesh->BodySetup->AggGeom;
+	AggGeom.ConvexElems.Reset();
+	FKConvexElem& ConvexElem = AggGeom.ConvexElems.AddDefaulted_GetRef();
+
+	ConvexElem.VertexData.Reserve( VertexPositions.Num() );
+	for ( const FVector& Position : VertexPositions )
+	{
+		ConvexElem.VertexData.Add( Position );
+	}
+
+	ConvexElem.UpdateElemBox();
+}
+
+void FDatasmithStaticMeshImporter::PreBuildStaticMeshes( FDatasmithImportContext& ImportContext )
+{
+	TRACE_CPUPROFILER_EVENT_SCOPE(FDatasmithStaticMeshImporter::PreBuildStaticMeshes);
+
+	FScopedSlowTask Progress(ImportContext.ImportedStaticMeshes.Num(), LOCTEXT("PreBuildStaticMeshes", "Packing UVs and computing tangents..."), true, *ImportContext.Warn);
+	Progress.MakeDialog(true);
+
+	IMeshUtilities& MeshUtilities = FModuleManager::Get().LoadModuleChecked< IMeshUtilities >( "MeshUtilities" );
+
+	TArray<UStaticMesh*> SortedMesh;
+	for (const TPair<TSharedRef<IDatasmithMeshElement>, UStaticMesh*>& Kvp : ImportContext.ImportedStaticMeshes)
+	{
+		if (Kvp.Value)
+		{
+			SortedMesh.Add(Kvp.Value);
+		}
+	}
+
+	// Start with the biggest mesh first to help balancing tasks on threads
+	SortedMesh.Sort(
+		[](const UStaticMesh& Lhs, const UStaticMesh& Rhs)
+		{
+			int32 LhsVerticesNum = Lhs.IsMeshDescriptionValid(0) ? Lhs.GetMeshDescription(0)->Vertices().Num() : 0;
+			int32 RhsVerticesNum = Rhs.IsMeshDescriptionValid(0) ? Rhs.GetMeshDescription(0)->Vertices().Num() : 0;
+
+			return LhsVerticesNum > RhsVerticesNum;
+		}
+	);
+
+	// Setup tasks in such a way that we're able to update the progress for tasks finishing in any order
+	TAtomic<uint32> TasksDone(0);
+	TArray<TFuture<bool>> TasksIsMeshValidResult;
+	for (UStaticMesh* StaticMesh : SortedMesh)
+	{
+		TasksIsMeshValidResult.Add(
+			Async(
+				EAsyncExecution::LargeThreadPool,
+				[StaticMesh]() { return PreBuildStaticMesh(StaticMesh); },
+				[&TasksDone]() { TasksDone++; }
+			)
+		);
+	}
+
+	// Ensure UI stays responsive by updating the progress even when the number of tasks hasn't changed
+	for (int32 OldTasksDone = 0, NewTasksDone = 0; OldTasksDone != SortedMesh.Num(); OldTasksDone = NewTasksDone)
+	{
+		NewTasksDone = TasksDone.Load();
+		Progress.EnterProgressFrame(NewTasksDone - OldTasksDone, FText::FromString(FString::Printf(TEXT("Packing UVs and computing tangents for static mesh %d/%d ..."), NewTasksDone, SortedMesh.Num())));
+		FPlatformProcess::Sleep(0.1);
+	}
+
+	for (int32 Index = 0; Index < TasksIsMeshValidResult.Num(); ++Index)
+	{
+		UStaticMesh* StaticMesh = SortedMesh[Index];
+
+		if (TasksIsMeshValidResult[Index].Get())
+		{
+			SortedMesh[Index]->MarkPackageDirty();
+		}
+		else
+		{
+			// Log invalid meshes and remove mesh reference from imported meshes and mesh actors
+			for (const TPair<TSharedRef<IDatasmithMeshElement>, UStaticMesh*>& Kvp : ImportContext.ImportedStaticMeshes)
+			{
+				if (Kvp.Value == StaticMesh)
+				{
+					ImportContext.LogError( FText::Format( LOCTEXT("PreBuildMesh", "Static mesh {0} is invalid since it contains only degenerated triangles or no triangle."), FText::FromName( StaticMesh->GetFName() ) ) );
+
+					ImportContext.ImportedStaticMeshes.Remove(Kvp.Key);
+					break;
+				}
+			}
+		}
+	}
+}
+
+class FMeshMaterialSlotBuilder
+{
+	struct SectionInfo
+	{
+		enum { InvalidIndex = -1 };
+
+		int32 LodIndex;
+		int32 SectionIndex;        // index of this group (in given LOD's mesh)
+		FPolygonGroupID SectionId; // id stored on the triangles
+		FName MaterialSlotName;    // per PolygonGroup FName attribute
+		int32 PolyCount;           // Number of poly
+		int32 StaticMaterialsIndex = InvalidIndex; // index in the FStaticMesh::StaticMaterials
+	};
+
+public:
+	FMeshMaterialSlotBuilder(UStaticMesh& StaticMesh, int32 MeshElementLODCount)
+	{
+		// Populate section LODs declared in IDatasmithMeshElement
+		for (int32 LodIndex = 0; LodIndex < MeshElementLODCount; ++LodIndex)
+		{
+			if (FMeshDescription* MeshDescription = StaticMesh.GetMeshDescription(LodIndex))
+			{
+				FStaticMeshAttributes Attributes(*MeshDescription);
+				TPolygonGroupAttributesConstRef<FName> MaterialSlotNameAttribute = Attributes.GetPolygonGroupMaterialSlotNames();
+				int32 SectionIndex = 0;
+				RawSections.Reserve(MeshDescription->PolygonGroups().Num());
+				for (FPolygonGroupID PolygonGroupID : MeshDescription->PolygonGroups().GetElementIDs())
+				{
+					SectionInfo& ThisSection = RawSections.Add_GetRef({});
+					ThisSection.LodIndex = LodIndex;
+					ThisSection.SectionIndex = SectionIndex++;
+					ThisSection.SectionId = PolygonGroupID;
+					ThisSection.MaterialSlotName = MaterialSlotNameAttribute[PolygonGroupID];
+					ThisSection.PolyCount = MeshDescription->GetPolygonGroupPolygons(PolygonGroupID).Num();
+				}
+			}
+		}
+
+		// Prepares StaticMaterials Slot informations from sections found in the mesh
+		for (SectionInfo& Section : RawSections)
+		{
+			// get associated slot (but skip empty polygon groups)
+			Section.StaticMaterialsIndex = Section.PolyCount <= 0 ? SectionInfo::InvalidIndex : SlotNames.Add(Section.MaterialSlotName).AsInteger();
+		}
+	}
+
+	FMeshSectionInfoMap GenerateSectionInfoMap() const
+	{
+		FMeshSectionInfoMap InfoMap;
+		for (const SectionInfo& Section : RawSections)
+		{
+			if (Section.StaticMaterialsIndex != SectionInfo::InvalidIndex)
+			{
+				InfoMap.Set(Section.LodIndex, Section.SectionIndex, FMeshSectionInfo(Section.StaticMaterialsIndex));
+			}
+		}
+		return InfoMap;
+	}
+
+	const TArray<SectionInfo>& GetSectionsInfos() const
+	{
+		return RawSections;
+	}
+
+	const TSet<FName>& GetSlotNames() const
+	{
+		return SlotNames;
+	}
+
+
+private:
+	TArray<SectionInfo> RawSections;
+	TSet<FName> SlotNames;
+};
+
+void FDatasmithStaticMeshImporter::SetupStaticMesh( FDatasmithAssetsImportContext& AssetsContext, TSharedRef< IDatasmithMeshElement > MeshElement, UStaticMesh* StaticMesh,
+	const FDatasmithStaticMeshImportOptions& StaticMeshImportOptions, float LightmapWeight )
+{
+	if (StaticMesh == nullptr)
+	{
+		return;
+	}
+
+	// The input static mesh must at least have its LOD 0 valid
+	if (!StaticMesh->IsMeshDescriptionValid(0))
+	{
+		return;
+	}
+
+	TRACE_CPUPROFILER_EVENT_SCOPE(FDatasmithStaticMeshImporter::SetupStaticMesh);
+
+	FMeshMaterialSlotBuilder Sections(*StaticMesh, MeshElement->GetLODCount());
+	UDatasmithStaticMeshTemplate* StaticMeshTemplate = NewObject< UDatasmithStaticMeshTemplate >( StaticMesh );
+
+	int32 MinLightmapSize = FDatasmithStaticMeshImportOptions::ConvertLightmapEnumToValue( StaticMeshImportOptions.MinLightmapResolution );
+	int32 MaxLightmapSize = FDatasmithStaticMeshImportOptions::ConvertLightmapEnumToValue( StaticMeshImportOptions.MaxLightmapResolution );
+
+	int32 LightmapSize = DatasmithStaticMeshImporterImpl::GetLightmapSize( LightmapWeight, MinLightmapSize, MaxLightmapSize );
+	StaticMesh->LightingGuid = FGuid::NewGuid();
+	StaticMesh->InitResources();
+
+	StaticMeshTemplate->LightMapResolution = LightmapSize;
+	StaticMeshTemplate->LightMapCoordinateIndex = 2;
+
+	// Get build requirement for the assigned materials
+	const int32 BuildRequirements = DatasmithStaticMeshImporterImpl::GetStaticMeshBuildRequirements( AssetsContext, MeshElement );
+
+	TSet< int32 > MaterialIndexToImportIndex;
+
+	bool bLightmapCoordinateIndexInvalid = false;
+
+	// Only setup LODs declared by the mesh element
+	for ( int32 LodIndex = 0; LodIndex < MeshElement->GetLODCount(); ++LodIndex)
+	{
+		if (!StaticMesh->IsMeshDescriptionValid(LodIndex))
+		{
+			continue;
+		}
+
+		FMeshDescription& MeshDescription = *StaticMesh->GetMeshDescription(LodIndex);
+
+		auto GetNextOpenUVChannel = [](FMeshDescription& Mesh, int32 FromUVChannelIndex)
+		{
+			for (int32 UVChannelIndex = FromUVChannelIndex; UVChannelIndex < MAX_MESH_TEXTURE_COORDS_MD; ++UVChannelIndex)
+			{
+				if (!DatasmithMeshHelper::HasUVData(Mesh, UVChannelIndex))
+				{
+					return UVChannelIndex;
+				}
+			}
+			return -1;
+		};
+
+		// UV Channels
+		int32 SourceIndex = -1;
+		bool bGenerateLightmapUVs = StaticMeshImportOptions.bGenerateLightmapUVs;
+		// Set DestinationIndex to 0 if generation of lightmap UVs is not required.
+		int32 DestinationIndex = bGenerateLightmapUVs ? -1 : 0;
+
+		// if a lightmap coordinate index was exported, disable lightmap generation
+		if (DatasmithMeshHelper::HasUVData(MeshDescription, MeshElement->GetLightmapCoordinateIndex()))
+		{
+			bGenerateLightmapUVs = false;
+		}
+
+		if (bGenerateLightmapUVs)
+		{
+			// Set the source lightmap index to the imported mesh data lightmap source if any, otherwise use the first open channel.
+			const bool bOriginalSourceIsValid = DatasmithMeshHelper::HasUVData(MeshDescription, MeshElement->GetLightmapSourceUV());
+			const int32 FirstOpenUVChannel = GetNextOpenUVChannel(MeshDescription, 0);
+
+			if (bOriginalSourceIsValid)
+			{
+				SourceIndex = MeshElement->GetLightmapSourceUV();
+				DestinationIndex = FirstOpenUVChannel;
+			}
+			else
+			{
+				//If the lightmap source index was not set, set it to 0 (as we will generate a UV latter if it's missing).
+				SourceIndex = 0;
+				DestinationIndex = FMath::Max(FirstOpenUVChannel, SourceIndex + 1);
+			}
+			
+			if (!FMath::IsWithin<int32>(SourceIndex, 0, MAX_MESH_TEXTURE_COORDS_MD))
+			{
+				AssetsContext.ParentContext.LogError(FText::Format(LOCTEXT("InvalidLightmapSourceIndexError", "Lightmap generation error for mesh {0}: Specified source is invalid {1}. Cannot find an available fallback source channel."), FText::FromName(StaticMesh->GetFName()), MeshElement->GetLightmapSourceUV()));
+				bGenerateLightmapUVs = false;
+			}
+			else if (!FMath::IsWithin<int32>(DestinationIndex, 0, MAX_MESH_TEXTURE_COORDS_MD))
+			{
+				AssetsContext.ParentContext.LogError(FText::Format(LOCTEXT("InvalidLightmapDestinationIndexError", "Lightmap generation error for mesh {0}: Cannot find an available destination channel."), FText::FromName(StaticMesh->GetFName())));
+				bGenerateLightmapUVs = false;
+			}
+
+			if (!bGenerateLightmapUVs)
+			{
+				AssetsContext.ParentContext.LogWarning(FText::Format(LOCTEXT("LightmapUVsWontBeGenerated", "Lightmap UVs for mesh {0} won't be generated. The lightmap coordinate index will be set to 0."), FText::FromName(StaticMesh->GetFName())));
+				SourceIndex = 0;
+				DestinationIndex = 1;
+			}
+		}
+
+		FDatasmithMeshBuildSettingsTemplate BuildSettingsTemplate;
+		BuildSettingsTemplate.Load(StaticMesh->GetSourceModel(LodIndex).BuildSettings);
+
+		BuildSettingsTemplate.bUseMikkTSpace = true;
+		BuildSettingsTemplate.bRecomputeNormals = FDatasmithStaticMeshImporter::ShouldRecomputeNormals(MeshDescription, BuildRequirements );
+		BuildSettingsTemplate.bRecomputeTangents = FDatasmithStaticMeshImporter::ShouldRecomputeTangents(MeshDescription, BuildRequirements );
+		BuildSettingsTemplate.bRemoveDegenerates = StaticMeshImportOptions.bRemoveDegenerates;
+		BuildSettingsTemplate.bUseHighPrecisionTangentBasis = true;
+		BuildSettingsTemplate.bUseFullPrecisionUVs = true;
+		BuildSettingsTemplate.bGenerateLightmapUVs = bGenerateLightmapUVs;
+		BuildSettingsTemplate.SrcLightmapIndex = SourceIndex;
+		BuildSettingsTemplate.DstLightmapIndex = DestinationIndex;
+		BuildSettingsTemplate.MinLightmapResolution = MinLightmapSize;
+
+		// Don't build adjacency buffer for meshes with over 500 000 triangles because it's too slow
+		BuildSettingsTemplate.bBuildAdjacencyBuffer = ((BuildRequirements & EMaterialRequirements::RequiresAdjacency) != 0) && DatasmithMeshHelper::GetPolygonCount(MeshDescription) < 500000;
+
+		StaticMeshTemplate->BuildSettings.Add( MoveTemp( BuildSettingsTemplate ) );
+
+		if (!StaticMeshImportOptions.bGenerateLightmapUVs)
+		{
+			StaticMeshTemplate->LightMapCoordinateIndex = MeshElement->GetLightmapCoordinateIndex() < 0 ? BuildSettingsTemplate.DstLightmapIndex : MeshElement->GetLightmapCoordinateIndex();
+		}
+		else if (MeshElement->GetLightmapCoordinateIndex() < 0)
+		{
+			StaticMeshTemplate->LightMapCoordinateIndex = bGenerateLightmapUVs ? BuildSettingsTemplate.DstLightmapIndex : BuildSettingsTemplate.SrcLightmapIndex;
+		}
+		else
+		{
+			int LastValidUVChannel = MAX_MESH_TEXTURE_COORDS_MD - 1;
+			while (LastValidUVChannel > 0 && !DatasmithMeshHelper::HasUVData(MeshDescription, LastValidUVChannel))
+			{
+				--LastValidUVChannel;
+			}
+
+			if (MeshElement->GetLightmapCoordinateIndex() <= LastValidUVChannel)
+			{
+				StaticMeshTemplate->LightMapCoordinateIndex = MeshElement->GetLightmapCoordinateIndex();
+			}
+			else
+			{
+				bLightmapCoordinateIndexInvalid = true;
+				StaticMeshTemplate->LightMapCoordinateIndex = LastValidUVChannel;
+			}
+		}
+	}
+
+	// Build a FMeshSectionInfoMap so that it computes the proper key for us
+	FMeshSectionInfoMap MeshSectionInfoMap = Sections.GenerateSectionInfoMap();
+	StaticMeshTemplate->SectionInfoMap.Load( MeshSectionInfoMap );
+
+	// Create one StaticMaterial per material index. Actual materials will be assigned later. This is just so that we create all the slots.
+	for (const FName& SlotName : Sections.GetSlotNames())
+	{
+		FDatasmithStaticMaterialTemplate StaticMaterial;
+		StaticMaterial.MaterialSlotName = SlotName;
+
+		StaticMeshTemplate->StaticMaterials.Add(MoveTemp(StaticMaterial));
+	}
+
+	ApplyMaterialsToStaticMesh(AssetsContext, MeshElement, StaticMeshTemplate);
+	StaticMeshTemplate->Apply(StaticMesh);
+
+	if (bLightmapCoordinateIndexInvalid)
+	{
+		FFormatNamedArguments FormatArgs;
+		FormatArgs.Add(TEXT("LightmapCoordinateIndex"), FText::FromString(FString::FromInt(MeshElement->GetLightmapCoordinateIndex())));
+		FormatArgs.Add(TEXT("MeshName"), FText::FromName(StaticMesh->GetFName()));
+		AssetsContext.ParentContext.LogError(FText::Format(LOCTEXT("InvalidLightmapSourceUVError", "The lightmap coordinate index '{LightmapCoordinateIndex}' used for the mesh '{MeshName}' is invalid. A valid lightmap coordinate index was set instead."), FormatArgs));
+	}
+}
+
+TMap< TSharedRef< IDatasmithMeshElement >, float > FDatasmithStaticMeshImporter::CalculateMeshesLightmapWeights( const TSharedRef< IDatasmithScene >& SceneElement )
+{
+	TMap< TSharedRef< IDatasmithMeshElement >, float > LightmapWeights;
+
+	float MaxArea = 0.0f;
+	float Max2DSurface = 0.0f;
+
+	// Compute the max values based on all meshes in the Datasmith Scene
+
+	for ( int32 MeshIndex = 0; MeshIndex < SceneElement->GetMeshesCount(); ++MeshIndex )
+	{
+		TSharedPtr< IDatasmithMeshElement > MeshElement = SceneElement->GetMesh( MeshIndex );
+
+		FVector4 Dimensions;
+		Dimensions.Set( MeshElement->GetWidth(), MeshElement->GetDepth(), MeshElement->GetHeight(), MeshElement->GetArea() );
+
+		MaxArea = FMath::Max( MaxArea, MeshElement->GetArea() );
+		Max2DSurface = FMath::Max(Max2DSurface, DatasmithStaticMeshImporterImpl::Get2DSurface( Dimensions ));
+	}
+
+	float MaxWeight = 0.0f;
+
+	for ( int32 MeshIndex = 0; MeshIndex < SceneElement->GetMeshesCount(); ++MeshIndex )
+	{
+		TSharedPtr< IDatasmithMeshElement > MeshElement = SceneElement->GetMesh( MeshIndex );
+
+		FVector4 Dimensions( MeshElement->GetWidth(), MeshElement->GetDepth(), MeshElement->GetHeight(), MeshElement->GetArea() );
+
+		MaxWeight = FMath::Max( MaxWeight, DatasmithStaticMeshImporterImpl::CalcBlendWeight( Dimensions, MaxArea, Max2DSurface ) );
+	}
+
+	for ( int32 MeshIndex = 0; MeshIndex < SceneElement->GetMeshesCount(); ++MeshIndex )
+	{
+		TSharedPtr< IDatasmithMeshElement > MeshElement = SceneElement->GetMesh( MeshIndex );
+
+		FVector4 Dimensions( MeshElement->GetWidth(), MeshElement->GetDepth(), MeshElement->GetHeight(), MeshElement->GetArea() );
+
+		float LightmapWeight = DatasmithStaticMeshImporterImpl::CalcBlendWeight( Dimensions, MaxArea, Max2DSurface ) / MaxWeight;
+
+		LightmapWeights.Add( MeshElement.ToSharedRef() ) = LightmapWeight;
+	}
+
+	return LightmapWeights;
+}
+
+void FDatasmithStaticMeshImporter::ApplyMaterialsToStaticMesh( const FDatasmithAssetsImportContext& AssetsContext, const TSharedRef< IDatasmithMeshElement >& MeshElement, UDatasmithStaticMeshTemplate* StaticMeshTemplate )
+{
+	for ( TPair<int32, FString>& Material : DatasmithStaticMeshImporterImpl::GetMaterials( MeshElement ) )
+	{
+		UMaterialInterface* MaterialInterface = FDatasmithImporterUtils::FindAsset< UMaterialInterface >( AssetsContext, *Material.Value );
+
+		FName MaterialSlotName = DatasmithMeshHelper::DefaultSlotName(Material.Key);
+
+		FDatasmithStaticMaterialTemplate* MaterialTemplate = Algo::FindByPredicate( StaticMeshTemplate->StaticMaterials, [ &MaterialSlotName ]( const FDatasmithStaticMaterialTemplate& Current ) -> bool
+		{
+			return Current.MaterialSlotName == MaterialSlotName;
+		});
+
+		if ( MaterialTemplate )
+		{
+			MaterialTemplate->MaterialInterface = MaterialInterface;
+		}
+	}
+}
+
+#undef LOCTEXT_NAMESPACE
diff --git a/Engine/Plugins/Enterprise/DatasmithImporter/Source/DatasmithImporter/Private/MasterMaterials/DatasmithMasterMaterial.cpp b/Engine/Plugins/Enterprise/DatasmithImporter/Source/DatasmithImporter/Private/MasterMaterials/DatasmithMasterMaterial.cpp
new file mode 100644
--- /dev/null
+++ b/Engine/Plugins/Enterprise/DatasmithImporter/Source/DatasmithImporter/Private/MasterMaterials/DatasmithMasterMaterial.cpp
@@ -0,0 +1,51 @@
+// Copyright 1998-2019 Epic Games, Inc. All Rights Reserved.
+
+#include "MasterMaterials/DatasmithMasterMaterial.h"
+
+#include "Materials/Material.h"
+#include "Materials/MaterialExpressionScalarParameter.h"
+#include "Materials/MaterialExpressionStaticBoolParameter.h"
+#include "Materials/MaterialExpressionTextureSampleParameter.h"
+#include "Materials/MaterialExpressionVectorParameter.h"
+
+#include "UObject/SoftObjectPath.h"
+
+FDatasmithMasterMaterial::FDatasmithMasterMaterial()
+	: Material( nullptr )
+{
+}
+
+void FDatasmithMasterMaterial::FromMaterial( UMaterial* InMaterial )
+{
+	if ( InMaterial )
+	{
+		for ( UMaterialExpression* Expression : InMaterial->Expressions )
+		{
+			FString ExpressionName = Expression->GetName();
+
+			if ( Expression->IsA< UMaterialExpressionVectorParameter >() )
+			{
+				VectorParams.Add( Expression->GetParameterName().ToString() );
+			}
+			else if ( Expression->IsA< UMaterialExpressionScalarParameter >() )
+			{
+				ScalarParams.Add( Expression->GetParameterName().ToString() );
+			}
+			else if ( Expression->IsA< UMaterialExpressionTextureSampleParameter >() )
+			{
+				TextureParams.Add( Expression->GetParameterName().ToString() );
+			}
+			else if ( Expression->IsA< UMaterialExpressionStaticBoolParameter >() )
+			{
+				BoolParams.Add( Expression->GetParameterName().ToString() );
+			}
+		}
+	}
+
+	Material = InMaterial;
+}
+
+void FDatasmithMasterMaterial::FromSoftObjectPath( const FSoftObjectPath& InObjectPath)
+{
+	FromMaterial( Cast< UMaterial >(InObjectPath.TryLoad() ) );
+}
diff --git a/Engine/Plugins/Enterprise/DatasmithImporter/Source/DatasmithImporter/Private/ObjectElements/DatasmithUElementsUtils.h b/Engine/Plugins/Enterprise/DatasmithImporter/Source/DatasmithImporter/Private/ObjectElements/DatasmithUElementsUtils.h
new file mode 100644
--- /dev/null
+++ b/Engine/Plugins/Enterprise/DatasmithImporter/Source/DatasmithImporter/Private/ObjectElements/DatasmithUElementsUtils.h
@@ -0,0 +1,164 @@
+// Copyright 1998-2019 Epic Games, Inc. All Rights Reserved.
+#pragma once
+
+#include "CoreMinimal.h"
+
+class FDatasmithUElementsUtils
+{
+public:
+	template<typename IElement, typename UElement, typename SetOp>
+	static UElement* FindOrAddElement(UObject* Outer, TMap<TWeakPtr<IElement>, UElement*>& Map, const TSharedPtr<IElement>& InElement, SetOp Assign)
+	{
+		if (InElement.IsValid())
+		{
+			UElement* FoundElement = Map.FindRef(InElement);
+			if (FoundElement)
+			{
+				return FoundElement;
+			}
+
+			FoundElement = NewObject<UElement>(Outer);
+			Assign(FoundElement);
+			Map.Add(InElement, FoundElement);
+			return FoundElement;
+		}
+
+		return nullptr;
+	}
+
+	/**
+	 * Utility to help iterating over the nested structure of the variant objects
+	 * Examples:
+	 * ForVariantElement<IDatasmithVariantElement>(SceneElement, [](TSharedPtr<IDatasmithVariantElement> Var)
+	 * {
+	 *     // Do something with all variants in the scene
+	 *
+	 *     return true; // It will keep iterating while the lambda returns true
+	 * });
+	 *
+	 * ForVariantElement<IDatasmithBasePropertyCaptureElement>(SomeVariantSet, [](TSharedPtr<IDatasmithBasePropertyCaptureElement> Prop)
+	 * {
+	 *     // Do something with all properties of a variant set
+	 *
+	 *     return true; // It will keep iterating while the lambda returns true
+	 * });
+	 */
+	template<typename TargetType, typename Func>
+	static inline typename TEnableIf<TIsSame<TargetType, IDatasmithLevelVariantSetsElement>::Value, bool>::Type ForVariantElement(const TSharedPtr<IDatasmithScene>& RootObject, Func Function)
+	{
+		for (int32 Index = 0; Index < RootObject->GetLevelVariantSetsCount(); ++Index)
+		{
+			bool bKeepIterating = Function(RootObject->GetLevelVariantSets(Index));
+			if (!bKeepIterating)
+			{
+				return false;
+			}
+		}
+		return true;
+	}
+	template<typename TargetType, typename Func>
+	static inline typename TEnableIf<!TIsSame<TargetType, IDatasmithLevelVariantSetsElement>::Value, bool>::Type ForVariantElement(const TSharedPtr<IDatasmithScene>& RootObject, Func Function)
+	{
+		for (int32 Index = 0; Index < RootObject->GetLevelVariantSetsCount(); ++Index)
+		{
+			bool bKeepIterating = ForVariantElement<TargetType>(RootObject->GetLevelVariantSets(Index), Function);
+			if (!bKeepIterating)
+			{
+				return false;
+			}
+		}
+		return true;
+	}
+	template<typename TargetType, typename Func>
+	static inline typename TEnableIf<TIsSame<TargetType, IDatasmithVariantSetElement>::Value, bool>::Type ForVariantElement(const TSharedPtr<IDatasmithLevelVariantSetsElement>& RootObject, Func Function)
+	{
+		for (int32 Index = 0; Index < RootObject->GetVariantSetsCount(); ++Index)
+		{
+			bool bKeepIterating = Function(RootObject->GetVariantSet(Index));
+			if (!bKeepIterating)
+			{
+				return false;
+			}
+		}
+		return true;
+	}
+	template<typename TargetType, typename Func>
+	static inline typename TEnableIf<!TIsSame<TargetType, IDatasmithVariantSetElement>::Value, bool>::Type ForVariantElement(const TSharedPtr<IDatasmithLevelVariantSetsElement>& RootObject, Func Function)
+	{
+		for (int32 Index = 0; Index < RootObject->GetVariantSetsCount(); ++Index)
+		{
+			bool bKeepIterating = ForVariantElement<TargetType>(RootObject->GetVariantSet(Index), Function);
+			if (!bKeepIterating)
+			{
+				return false;
+			}
+		}
+		return true;
+	}
+
+	template<typename TargetType, typename Func>
+	static inline typename TEnableIf<TIsSame<TargetType, IDatasmithVariantElement>::Value, bool>::Type ForVariantElement(const TSharedPtr<IDatasmithVariantSetElement>& RootObject, Func Function)
+	{
+		for (int32 Index = 0; Index < RootObject->GetVariantsCount(); ++Index)
+		{
+			bool bKeepIterating = Function(RootObject->GetVariant(Index));
+			if (!bKeepIterating)
+			{
+				return false;
+			}
+		}
+		return true;
+	}
+	template<typename TargetType, typename Func>
+	static inline typename TEnableIf<!TIsSame<TargetType, IDatasmithVariantElement>::Value, bool>::Type ForVariantElement(const TSharedPtr<IDatasmithVariantSetElement>& RootObject, Func Function)
+	{
+		for (int32 Index = 0; Index < RootObject->GetVariantsCount(); ++Index)
+		{
+			bool bKeepIterating = ForVariantElement<TargetType>(RootObject->GetVariant(Index), Function);
+			if (!bKeepIterating)
+			{
+				return false;
+			}
+		}
+		return true;
+	}
+	template<typename TargetType, typename Func>
+	static inline typename TEnableIf<TIsSame<TargetType, IDatasmithActorBindingElement>::Value, bool>::Type ForVariantElement(const TSharedPtr<IDatasmithVariantElement>& RootObject, Func Function)
+	{
+		for (int32 Index = 0; Index < RootObject->GetActorBindingsCount(); ++Index)
+		{
+			bool bKeepIterating = Function(RootObject->GetActorBinding(Index));
+			if (!bKeepIterating)
+			{
+				return false;
+			}
+		}
+		return true;
+	}
+	template<typename TargetType, typename Func>
+	static inline typename TEnableIf<!TIsSame<TargetType, IDatasmithActorBindingElement>::Value, bool>::Type ForVariantElement(const TSharedPtr<IDatasmithVariantElement>& RootObject, Func Function)
+	{
+		for (int32 Index = 0; Index < RootObject->GetActorBindingsCount(); ++Index)
+		{
+			bool bKeepIterating = ForVariantElement<TargetType>(RootObject->GetActorBinding(Index), Function);
+			if (!bKeepIterating)
+			{
+				return false;
+			}
+		}
+		return true;
+	}
+	template<typename TargetType, typename Func>
+	static inline typename TEnableIf<TIsSame<TargetType, IDatasmithBasePropertyCaptureElement>::Value, bool>::Type ForVariantElement(const TSharedPtr<IDatasmithActorBindingElement>& RootObject, Func Function)
+	{
+		for (int32 Index = 0; Index < RootObject->GetPropertyCapturesCount(); ++Index)
+		{
+			bool bKeepIterating = Function(RootObject->GetPropertyCapture(Index));
+			if (!bKeepIterating)
+			{
+				return false;
+			}
+		}
+		return true;
+	}
+};
diff --git a/Engine/Plugins/Enterprise/DatasmithImporter/Source/DatasmithImporter/Private/Translators/DatasmithPayload.cpp b/Engine/Plugins/Enterprise/DatasmithImporter/Source/DatasmithImporter/Private/Translators/DatasmithPayload.cpp
new file mode 100644
--- /dev/null
+++ b/Engine/Plugins/Enterprise/DatasmithImporter/Source/DatasmithImporter/Private/Translators/DatasmithPayload.cpp
@@ -0,0 +1,3 @@
+// Copyright 1998-2019 Epic Games, Inc. All Rights Reserved.
+
+#include "Translators/DatasmithPayload.h"
diff --git a/Engine/Plugins/Enterprise/DatasmithImporter/Source/DatasmithImporter/Private/Translators/DatasmithTranslatableSource.h b/Engine/Plugins/Enterprise/DatasmithImporter/Source/DatasmithImporter/Private/Translators/DatasmithTranslatableSource.h
new file mode 100644
--- /dev/null
+++ b/Engine/Plugins/Enterprise/DatasmithImporter/Source/DatasmithImporter/Private/Translators/DatasmithTranslatableSource.h
@@ -0,0 +1,34 @@
+// Copyright 1998-2019 Epic Games, Inc. All Rights Reserved.
+
+#pragma once
+
+#include "DatasmithSceneSource.h"
+
+#include "CoreMinimal.h"
+#include "UObject/StrongObjectPtr.h"
+
+/**
+ * Wrap a source with an adapted translator.
+ * This scopes the lifecycle of a translator, 
+ */
+struct DATASMITHIMPORTER_API FDatasmithTranslatableSceneSource
+{
+	FDatasmithTranslatableSceneSource(const FDatasmithSceneSource& Source);
+	~FDatasmithTranslatableSceneSource();
+
+	bool IsTranslatable() const;
+
+	bool Translate(TSharedRef< class IDatasmithScene > Scene);
+
+	TSharedPtr<class IDatasmithTranslator> GetTranslator() const;
+
+private:
+	/** Translator currently in use (null when not importing) */
+	TSharedPtr<class IDatasmithTranslator> Translator;
+
+	/** internal helper to release scene */
+	TUniquePtr< class FSceneGuard > SceneGuard;
+};
+
+
+
