Copying //UE4/Dev-Enterprise @ cl 9420543 to Dev-Main (//UE4/Dev-Main)

diff --git a/Engine/Plugins/Enterprise/DatasmithImporter/Source/DatasmithImporter/Private/UI/DatasmithUICommands.cpp b/Engine/Plugins/Enterprise/DatasmithImporter/Source/DatasmithImporter/Private/UI/DatasmithUICommands.cpp
new file mode 100644
--- /dev/null
+++ b/Engine/Plugins/Enterprise/DatasmithImporter/Source/DatasmithImporter/Private/UI/DatasmithUICommands.cpp
@@ -0,0 +1,49 @@
+// Copyright 1998-2019 Epic Games, Inc. All Rights Reserved.
+
+#include "UI/DatasmithUICommands.h"
+#include "UI/DatasmithStyle.h"
+#include "Styling/ISlateStyle.h"
+
+#define LOCTEXT_NAMESPACE "DatasmithImporter"
+
+FDatasmithUICommands::FDatasmithUICommands()
+: TCommands<FDatasmithUICommands>(FDatasmithStyle::GetContextName(), LOCTEXT("DatasmithImporter", "Datasmith Importer"), NAME_None, FDatasmithStyle::GetStyleSetName())
+{
+}
+
+void FDatasmithUICommands::RegisterCommands()
+{
+	UI_COMMAND(RepeatLastImport, "Import", "Repeat last import", EUserInterfaceActionType::Button, FInputChord());
+}
+
+TSharedPtr< FUICommandInfo > FDatasmithUICommands::AddMenuCommand(const FString& CommandName, const FText& Caption, const FText& Description)
+{
+	// Get the non-const instance so that we can add a new command to the list
+	FDatasmithUICommands& DatasmithCommands = *(Instance.Pin());
+
+	TSharedPtr<FUICommandInfo>& Command = DatasmithCommands.MenuCommands.AddDefaulted_GetRef();
+	const FString DotString = FString(TEXT(".")) + CommandName;
+
+	FUICommandInfo::MakeCommandInfo(
+		DatasmithCommands.AsShared(),
+		Command,
+		*CommandName,
+		Caption,
+		Description,
+		FSlateIcon(DatasmithCommands.GetStyleSetName(), ISlateStyle::Join(DatasmithCommands.GetContextName(), TCHAR_TO_ANSI(*DotString))),
+		EUserInterfaceActionType::Check,
+		FInputChord()
+		);
+
+	return Command;
+}
+
+void FDatasmithUICommands::RemoveMenuCommand(const TSharedPtr< FUICommandInfo >& Command)
+{
+	// Get the non-const instance so that we can remove the command from the list
+	FDatasmithUICommands& DatasmithCommands = *(Instance.Pin());
+
+	DatasmithCommands.MenuCommands.RemoveSingle(Command);
+}
+
+#undef LOCTEXT_NAMESPACE
diff --git a/Engine/Plugins/Enterprise/DatasmithImporter/Source/DatasmithImporter/Private/UI/DatasmithUIManager.cpp b/Engine/Plugins/Enterprise/DatasmithImporter/Source/DatasmithImporter/Private/UI/DatasmithUIManager.cpp
new file mode 100644
--- /dev/null
+++ b/Engine/Plugins/Enterprise/DatasmithImporter/Source/DatasmithImporter/Private/UI/DatasmithUIManager.cpp
@@ -0,0 +1,149 @@
+// Copyright 1998-2019 Epic Games, Inc. All Rights Reserved.
+
+#include "UI/DatasmithUIManager.h"
+
+#include "UI/DatasmithUICommands.h"
+#include "UI/DatasmithStyle.h"
+#include "LevelEditor.h"
+#include "Modules/ModuleManager.h"
+
+#include "Framework/MultiBox/MultiBoxBuilder.h"
+
+#define LOCTEXT_NAMESPACE "DatasmithImporter"
+
+TUniquePtr<FDatasmithUIManager> FDatasmithUIManager::Instance = NULL;
+
+void FDatasmithUIManager::Initialize()
+{
+	if (!Instance.IsValid())
+	{
+		Instance = MakeUnique<FDatasmithUIManager>();
+
+		FDatasmithStyle::Initialize();
+		FDatasmithUICommands::Register();
+
+		Instance->ExtendToolbar();
+	}
+}
+
+void FDatasmithUIManager::Shutdown()
+{
+	if (Instance.IsValid())
+	{
+		FDatasmithUICommands::Unregister();
+		FDatasmithStyle::Shutdown();
+
+		Instance.Reset();
+	}
+}
+
+FDatasmithUIManager& FDatasmithUIManager::Get()
+{
+	check(Instance);
+	return *Instance;
+}
+
+void FDatasmithUIManager::ExtendToolbar()
+{
+	DatasmithActions = MakeShareable(new FUICommandList);
+
+	// Action to repeat the last selected command
+	DatasmithActions->MapAction(FDatasmithUICommands::Get().RepeatLastImport,
+		FExecuteAction::CreateLambda([this]() { DatasmithActions->ExecuteAction(GetLastSelectedCommand()); })
+	);
+
+	// Add a Datasmith toolbar section after the settings section of the level editor
+	FLevelEditorModule& LevelEditorModule = FModuleManager::LoadModuleChecked<FLevelEditorModule>("LevelEditor");
+
+	TSharedPtr<FExtender> ToolbarExtender = MakeShareable(new FExtender);
+	ToolbarExtender->AddToolBarExtension("Settings", EExtensionHook::After, DatasmithActions, FToolBarExtensionDelegate::CreateRaw(this, &FDatasmithUIManager::FillToolbar));
+
+	LevelEditorModule.GetToolBarExtensibilityManager()->AddExtender(ToolbarExtender);
+}
+
+TSharedRef<SWidget> GenerateDatasmithMenuContent(const TSharedPtr<FUICommandList>& InCommandList)
+{
+	const bool bShouldCloseWindowAfterMenuSelection = true;
+	FMenuBuilder MenuBuilder(bShouldCloseWindowAfterMenuSelection, InCommandList);
+
+	// Fill the Datasmith importers section in the drop-down menu with the registered importers
+	MenuBuilder.BeginSection("DatasmithImporters", LOCTEXT("DatasmithImportersSection", "Datasmith Importers"));
+	{
+		for (const TSharedPtr<FUICommandInfo>& Command : FDatasmithUICommands::Get().MenuCommands)
+		{
+			MenuBuilder.AddMenuEntry(Command);
+		}
+	}
+	MenuBuilder.EndSection();
+
+	return MenuBuilder.MakeWidget();
+}
+
+void FDatasmithUIManager::FillToolbar(FToolBarBuilder& ToolbarBuilder)
+{
+	ToolbarBuilder.BeginSection("Datasmith");
+	{
+		// Add a button to repeat the last selected import command
+		// The button's text and visual content is fetched from the last selected menu entry
+		ToolbarBuilder.AddToolBarButton(
+			FDatasmithUICommands::Get().RepeatLastImport,
+			NAME_None,
+			TAttribute< FText >::Create(TAttribute< FText >::FGetter::CreateLambda([this]() { return GetLastSelectedCommand()->GetLabel(); })),
+			TAttribute< FText >::Create(TAttribute< FText >::FGetter::CreateLambda([this]() { return GetLastSelectedCommand()->GetDescription(); })),
+			TAttribute< FSlateIcon >::Create(TAttribute< FSlateIcon >::FGetter::CreateLambda([this]() { return GetLastSelectedCommand()->GetIcon(); }))
+		);
+	}
+	ToolbarBuilder.EndSection();
+}
+
+TSharedPtr<FUICommandInfo> FDatasmithUIManager::AddMenuEntry(const FString& CommandName, const FText& Caption, const FText& Description, const FString& IconResourcePath, FExecuteAction ExecuteAction, UClass* FactoryClass)
+{
+	// Setup the icon and command. They are binded together through the CommandName
+	FDatasmithStyle::SetIcon(CommandName, IconResourcePath);
+	TSharedPtr<FUICommandInfo> Command = FDatasmithUICommands::AddMenuCommand(CommandName, Caption, Description);
+
+	FactoryClassToUICommandMap.Add(FactoryClass, Command);
+
+	// The menu entry will have a check mark next to it if it was the last selected command
+	DatasmithActions->MapAction(Command,
+		ExecuteAction,
+		FCanExecuteAction(),
+		FIsActionChecked::CreateLambda([=]() { return LastSelectedCommand == Command; })
+	);
+	return Command;
+}
+
+void FDatasmithUIManager::RemoveMenuEntry(const TSharedPtr<FUICommandInfo>& Command)
+{
+	if (LastSelectedCommand == Command)
+	{
+		LastSelectedCommand.Reset();
+	}
+
+	UClass* const* Class = FactoryClassToUICommandMap.FindKey(Command);
+	if (Class)
+	{
+		FactoryClassToUICommandMap.Remove(*Class);
+	}
+
+	DatasmithActions->UnmapAction(Command);
+	FDatasmithUICommands::RemoveMenuCommand(Command);
+}
+
+const TSharedRef<FUICommandInfo> FDatasmithUIManager::GetLastSelectedCommand()
+{
+	if (!LastSelectedCommand.IsValid())
+	{
+		// Initialize the last selected command with the first command, which is always the UDatasmith import since it's the first to get loaded and registered
+		LastSelectedCommand = FDatasmithUICommands::Get().MenuCommands[0];
+	}
+	return LastSelectedCommand.ToSharedRef();
+}
+
+void FDatasmithUIManager::SetLastFactoryUsed(UClass* Class)
+{
+	// Internally, we only need to know the command associated with the last factory used
+	LastSelectedCommand = FactoryClassToUICommandMap[Class];
+}
+
+#undef LOCTEXT_NAMESPACE
diff --git a/Engine/Plugins/Enterprise/DatasmithImporter/Source/DatasmithImporter/Private/Utility/DatasmithImportFactoryHelper.h b/Engine/Plugins/Enterprise/DatasmithImporter/Source/DatasmithImporter/Private/Utility/DatasmithImportFactoryHelper.h
new file mode 100644
--- /dev/null
+++ b/Engine/Plugins/Enterprise/DatasmithImporter/Source/DatasmithImporter/Private/Utility/DatasmithImportFactoryHelper.h
@@ -0,0 +1,21 @@
+// Copyright 1998-2019 Epic Games, Inc. All Rights Reserved.
+
+#pragma once
+
+#include "CoreMinimal.h"
+#include "AssetData.h"
+
+class USceneComponent;
+class FJsonObject;
+class FString;
+class AActor;
+struct FBox;
+
+namespace DatasmithImportFactoryHelper
+{
+	TSharedPtr<FJsonObject> LoadJsonFile(const FString& JsonFilename);
+
+	void ComputeBounds( USceneComponent* ActorComponent, FBox& Bounds );
+
+	void CaptureSceneThumbnail(AActor* SceneActor, TArray<FAssetData>& AssetDataList);
+}
diff --git a/Engine/Plugins/Enterprise/DatasmithImporter/Source/DatasmithImporter/Private/Utility/DatasmithMathUtils.cpp b/Engine/Plugins/Enterprise/DatasmithImporter/Source/DatasmithImporter/Private/Utility/DatasmithMathUtils.cpp
new file mode 100644
--- /dev/null
+++ b/Engine/Plugins/Enterprise/DatasmithImporter/Source/DatasmithImporter/Private/Utility/DatasmithMathUtils.cpp
@@ -0,0 +1,35 @@
+// Copyright 1998-2019 Epic Games, Inc. All Rights Reserved.
+
+#include "Utility/DatasmithMathUtils.h"
+
+#include "Math/UnrealMathUtility.h"
+#include "Math/TransformVectorized.h"
+#include "Math/Quat.h"
+
+
+void FDatasmithTransformUtils::GetRotation(const FTransform& InTransform, FQuat& OutRotation)
+{
+	OutRotation = InTransform.GetRotation();
+
+	// Workaround to floating point precision issues with quaternion
+	{
+		// this value was found from ONE experience, so I'm waiting on other problem to improve it
+		if (FMath::Abs(OutRotation.X) < 0.005f)
+		{
+			OutRotation.X = 0.f;
+		}
+		if (FMath::Abs(OutRotation.Y) < 0.005f)
+		{
+			OutRotation.Y = 0.f;
+		}
+		if (FMath::Abs(OutRotation.Z) < 0.005f)
+		{
+			OutRotation.Z = 0.f;
+		}
+		if ((1 - FMath::Abs(OutRotation.W)) < 0.001f)
+		{
+			OutRotation.W = (OutRotation.W > 0) ? 1.f : -1.f;
+		}
+		OutRotation.Normalize();
+	}
+}
\ No newline at end of file
diff --git a/Engine/Plugins/Enterprise/DatasmithImporter/Source/DatasmithImporter/Public/ActorFactoryDatasmithScene.h b/Engine/Plugins/Enterprise/DatasmithImporter/Source/DatasmithImporter/Public/ActorFactoryDatasmithScene.h
new file mode 100644
--- /dev/null
+++ b/Engine/Plugins/Enterprise/DatasmithImporter/Source/DatasmithImporter/Public/ActorFactoryDatasmithScene.h
@@ -0,0 +1,23 @@
+// Copyright 1998-2019 Epic Games, Inc. All Rights Reserved.
+
+#pragma once
+
+#include "ActorFactories/ActorFactory.h"
+
+#include "ActorFactoryDatasmithScene.generated.h"
+
+UCLASS(MinimalAPI)
+class UActorFactoryDatasmithScene : public UActorFactory
+{
+	GENERATED_UCLASS_BODY()
+
+public:
+	static void SpawnRelatedActors( class ADatasmithSceneActor* DatasmithSceneActor, bool bReimportDeletedActors );
+
+	//~ Begin UActorFactory Interface
+	virtual bool CanCreateActorFrom( const FAssetData& AssetData, FText& OutErrorMsg ) override;
+
+protected:
+	virtual AActor* SpawnActor( UObject* Asset, ULevel* InLevel, const FTransform& Transform, EObjectFlags InObjectFlags, const FName Name ) override;
+	//~ End UActorFactory Interface
+};
\ No newline at end of file
diff --git a/Engine/Plugins/Enterprise/DatasmithImporter/Source/DatasmithImporter/Public/DatasmithLightImporter.h b/Engine/Plugins/Enterprise/DatasmithImporter/Source/DatasmithImporter/Public/DatasmithLightImporter.h
new file mode 100644
--- /dev/null
+++ b/Engine/Plugins/Enterprise/DatasmithImporter/Source/DatasmithImporter/Public/DatasmithLightImporter.h
@@ -0,0 +1,63 @@
+// Copyright 1998-2019 Epic Games, Inc. All Rights Reserved.
+
+#pragma once
+
+#include "Containers/ContainersFwd.h"
+#include "Templates/SharedPointer.h"
+
+class AActor;
+class ADatasmithAreaLightActor;
+class ALight;
+class ALightmassPortal;
+class AGroupActor;
+class APointLight;
+class AStaticMeshActor;
+struct FDatasmithImportContext;
+class FString;
+class IDatasmithAreaLightElement;
+class IDatasmithLightActorElement;
+class IDatasmithLightmassPortalElement;
+class IDatasmithPointLightElement;
+class IDatasmithShaderElement;
+class IDatasmithSpotLightElement;
+class IDatasmithTextureElement;
+struct FDatasmithImportContext;
+class UDatasmithAreaLightComponent;
+class ULightComponent;
+class ULightmassPortalComponent;
+class UPackage;
+class UPointLightComponent;
+class USceneComponent;
+class USpotLightComponent;
+class UTexture;
+class UTextureLightProfile;
+class UWorld;
+
+class DATASMITHIMPORTER_API FDatasmithLightImporter
+{
+public:
+	static AActor* ImportLightActor( const TSharedRef< IDatasmithLightActorElement >& LightElement, FDatasmithImportContext& ImportContext );
+	static USceneComponent* ImportLightComponent( const TSharedRef< IDatasmithLightActorElement >& LightElement, FDatasmithImportContext& ImportContext, UObject* Outer );
+
+	static AActor* CreateHDRISkyLight(const TSharedPtr< IDatasmithShaderElement >& ShaderElement, FDatasmithImportContext& ImportContext);
+
+	static AActor* CreatePhysicalSky(FDatasmithImportContext& ImportContext);
+
+private:
+	static void CreateIESTexture(FDatasmithImportContext& InContext, const TSharedPtr< IDatasmithLightActorElement >& InLightActorElement);
+	static void SetTextureLightProfile( const TSharedRef< IDatasmithLightActorElement >& LightElement, class UDatasmithLightComponentTemplate* LightComponentTemplate, const TCHAR* LightsFolderPath );
+	static UTextureLightProfile* FindTextureLightProfile( const TSharedRef< IDatasmithLightActorElement >& LightElement, const TCHAR* LightsFolderPath );
+
+	static void SetupLightComponent( ULightComponent* LightComponent, const TSharedPtr< IDatasmithLightActorElement >& LightElement, const TCHAR* MaterialsFolderPath, const TCHAR* LightsFolderPath );
+	static void SetupPointLightComponent( UPointLightComponent* PointLightComponent, const TSharedRef< IDatasmithPointLightElement >& PointLightElement, const TCHAR* LightsFolderPath, const TCHAR* MaterialsFolderPath );
+	static void SetupSpotLightComponent( USpotLightComponent* SpotLightComponent, const TSharedRef< IDatasmithSpotLightElement >& SpotLightElement, const TCHAR* LightsFolderPath, const TCHAR* MaterialsFolderPath );
+
+	static AActor* ImportAreaLightActor( const TSharedRef< IDatasmithAreaLightElement >& AreaLightElement, FDatasmithImportContext& ImportContext );
+	static USceneComponent* ImportAreaLightComponent( const TSharedRef< IDatasmithAreaLightElement >& AreaLightElement, FDatasmithImportContext& ImportContext, UObject* Outer );
+
+	static AActor* CreateAreaLightActor( const TSharedRef< IDatasmithAreaLightElement >& AreaLightElement, FDatasmithImportContext& ImportContext );
+	static void SetupAreaLightActor( const TSharedRef< IDatasmithAreaLightElement >& AreaLightElement, FDatasmithImportContext& ImportContext, ADatasmithAreaLightActor* LightShapeActor );
+	static ULightmassPortalComponent* ImportLightmassPortalComponent( const TSharedRef< IDatasmithLightmassPortalElement >& LightElement, FDatasmithImportContext& ImportContext, UObject* Outer );
+	static AActor* CreateSkyLight(const TSharedPtr< IDatasmithShaderElement >& ShaderElement, FDatasmithImportContext& ImportContext, bool bUseHDRMat);
+
+};
diff --git a/Engine/Plugins/Enterprise/DatasmithImporter/Source/DatasmithImporter/Public/DatasmithMeshHelper.h b/Engine/Plugins/Enterprise/DatasmithImporter/Source/DatasmithImporter/Public/DatasmithMeshHelper.h
new file mode 100644
--- /dev/null
+++ b/Engine/Plugins/Enterprise/DatasmithImporter/Source/DatasmithImporter/Public/DatasmithMeshHelper.h
@@ -0,0 +1,148 @@
+// Copyright 1998-2019 Epic Games, Inc. All Rights Reserved.
+
+#pragma once
+
+#include "CoreTypes.h"
+#include "Containers/ContainersFwd.h"
+#include "Math/Vector.h"
+
+class FDatasmithMesh;
+class FMD5;
+class FName;
+class UStaticMesh;
+struct FMeshDescription;
+struct FMeshTriangle;
+struct FRawMesh;
+struct FVector;
+
+namespace DatasmithMeshHelper
+{
+	/**
+	 * Number of polygon in the given Mesh
+	 *
+	 * @param Mesh      Source mesh
+	 * @return int32    number of polys
+	 */
+	DATASMITHIMPORTER_API int32 GetPolygonCount(const FMeshDescription& Mesh);
+
+	/**
+	 * Number of triangles in the given Mesh
+	 *
+	 * @param Mesh      Source mesh
+	 * @return int32    number of triangles
+	 */
+	DATASMITHIMPORTER_API int32 GetTriangleCount(const FMeshDescription& Mesh);
+
+	/**
+	 * Build a copy of all positions.
+	 * Can be used to build a convexHull, a bounding box, etc...
+	 *
+	 * @param Mesh
+	 * @param Positions
+	 * @return DATASMITHIMPORTER_API void
+	 */
+	DATASMITHIMPORTER_API void ExtractVertexPositions(const FMeshDescription& Mesh, TArray<FVector>& Positions);
+
+	/**
+	 * Register all attributes required to correctly interact with StaticMesh.
+	 * @note: Delegates to FStaticMesh the registering of standard attributes
+	 *
+	 * @param Mesh   MeshDescription to prepare for StaticMesh usage
+	 */
+	DATASMITHIMPORTER_API void PrepareAttributeForStaticMesh(FMeshDescription& Mesh);
+
+	/**
+	 * Test if a triangle is extremely small or if its corners are aligned.
+	 *
+	 * @param Mesh		    Source mesh that contains the tested triangle
+	 * @param MeshTriangle	Triangle in that mesh
+	 * @return bool         True for degenerated triangles, false otherwise
+	 */
+	DATASMITHIMPORTER_API bool IsTriangleDegenerated(const FMeshDescription& Mesh, const FMeshTriangle& MeshTriangle);
+
+	/**
+	 * Write mesh data on a UStaticMesh.
+	 * RawMesh will be converted to MeshDescription beforehand.
+	 *
+	 * @param StaticMesh    Destination StaticMesh
+	 * @param LodIndex      Destination index in the StaticMesh
+	 * @param Mesh          Source data to write into the StaticMesh
+	 */
+	DATASMITHIMPORTER_API void FillUStaticMesh(UStaticMesh* StaticMesh, int32 LodIndex, FRawMesh& RawMesh, const TMap<int32, FName>* InMaterialMapInverse = nullptr);
+	DATASMITHIMPORTER_API void FillUStaticMesh(UStaticMesh* StaticMesh, int32 LodIndex, FDatasmithMesh& DSMesh);
+	DATASMITHIMPORTER_API void FillUStaticMeshByCopy(UStaticMesh* StaticMesh, int LodIndex, FMeshDescription MeshDescription);
+	DATASMITHIMPORTER_API void FillUStaticMesh(UStaticMesh* StaticMesh, int LodIndex, FMeshDescription&& MeshDescription);
+	/**
+	 * Make sure the StaticMesh can host all used materials
+	 *
+	 * @param				    StaticMesh Mesh to prepare
+	 * @param MaterialCount     Number of slot required
+	 */
+	DATASMITHIMPORTER_API void PrepareStaticMaterials(UStaticMesh* StaticMesh, int32 MaterialCount);
+
+	/**
+	 * Used in FRawMesh <-> FMeshDescription context, create the SlotName <-> MaterialIndex mapping
+	 * MaterialIndex in range [0, StaticMaterials.Num()-1]
+	 * SlotName extracted from, StaticMaterials.ImportedMaterialSlotName
+	 *
+	 * @param StaticMesh            Source of the mapping information
+	 * @param MaterialMap	        map SlotName to MaterialIndex
+	 * @param MaterialMapInverse    map a MaterialIndex to SlotName
+	 */
+	DATASMITHIMPORTER_API void BuildMaterialMappingFromStaticMesh(UStaticMesh* StaticMesh, TMap<FName, int32>& MaterialMap, TMap<int32, FName>& MaterialMapInverse);
+
+	/**
+	 * Hash a Mesh at specified LOD index in a StaticMesh
+	 *
+	 * @param StaticMesh    Source Mesh set
+	 * @param LodIndex      LOD index of the Mesh to hash
+	 * @param MD5           Hash to update
+	 * @return DATASMITHIMPORTER_API void
+	 */
+	DATASMITHIMPORTER_API void HashMeshLOD(UStaticMesh* StaticMesh, int32 LodIndex, FMD5& MD5);
+
+	/**
+	 * Hash a MeshDescription
+	 *
+	 * @param Mesh  MeshDescription to hash
+	 * @param MD5   Hash to update
+	 */
+	DATASMITHIMPORTER_API void HashMeshDescription(const FMeshDescription& Mesh, FMD5& MD5);
+
+	/**
+	 * Generate a default FName to name a material slot based on the index
+	 *
+	 * @param MaterialIndex index of the slot to name
+	 */
+	DATASMITHIMPORTER_API FName DefaultSlotName(int32 MaterialIndex);
+
+	DATASMITHIMPORTER_API int32 GetNumUVChannel(const FMeshDescription& Mesh);
+	DATASMITHIMPORTER_API bool HasUVChannel(const FMeshDescription& Mesh, int32 ChannelIndex);
+
+	DATASMITHIMPORTER_API bool HasUVData(const FMeshDescription& Mesh, int32 ChannelIndex);
+
+	DATASMITHIMPORTER_API bool RequireUVChannel(FMeshDescription& Mesh, int32 ChannelIndex);
+
+	/**
+	 * Generate simple UV data at channel 0.
+	 *
+	 * @param DatasmithMesh		The DatasmithMesh in which the UV data will be created.
+	 */
+	DATASMITHIMPORTER_API void CreateDefaultUVs(FDatasmithMesh& DatasmithMesh);
+
+	/**
+	 * Generate simple UV data at channel 0.
+	 *
+	 * @param MeshDescription	The MeshDescription in which the UV data will be created.
+	 */
+	DATASMITHIMPORTER_API void CreateDefaultUVs(FMeshDescription& MeshDescription);
+
+	/**
+	 * Verify that at least one triangle in that mesh is not degenerated
+	 *
+	 * @param Mesh
+	 * @param BuildScale
+	 * @return DATASMITHIMPORTER_API bool
+	 */
+	DATASMITHIMPORTER_API bool IsMeshValid(const FMeshDescription& Mesh, FVector BuildScale=FVector::OneVector);
+}
diff --git a/Engine/Plugins/Enterprise/DatasmithImporter/Source/DatasmithImporter/Public/MasterMaterials/DatasmithMasterMaterial.h b/Engine/Plugins/Enterprise/DatasmithImporter/Source/DatasmithImporter/Public/MasterMaterials/DatasmithMasterMaterial.h
new file mode 100644
--- /dev/null
+++ b/Engine/Plugins/Enterprise/DatasmithImporter/Source/DatasmithImporter/Public/MasterMaterials/DatasmithMasterMaterial.h
@@ -0,0 +1,31 @@
+// Copyright 1998-2019 Epic Games, Inc. All Rights Reserved.
+#pragma once
+
+#include "Containers/Set.h"
+#include "Containers/UnrealString.h"
+
+struct FSoftObjectPath;
+class UMaterial;
+
+class DATASMITHIMPORTER_API FDatasmithMasterMaterial
+{
+public:
+	FDatasmithMasterMaterial();
+	explicit FDatasmithMasterMaterial( UMaterial* InMaterial );
+
+	void FromMaterial( UMaterial* InMaterial );
+	void FromSoftObjectPath( const FSoftObjectPath& InObjectPath );
+
+	UMaterial* GetMaterial() const { return Material; }
+
+	bool IsValid() const { return Material != nullptr; }
+
+	TSet< FString > VectorParams;
+	TSet< FString > ScalarParams;
+	TSet< FString > TextureParams;
+	TSet< FString > BoolParams;
+
+private:
+	UMaterial* Material;
+
+};
diff --git a/Engine/Plugins/Enterprise/DatasmithImporter/Source/DatasmithImporter/Public/MasterMaterials/DatasmithMasterMaterialManager.h b/Engine/Plugins/Enterprise/DatasmithImporter/Source/DatasmithImporter/Public/MasterMaterials/DatasmithMasterMaterialManager.h
new file mode 100644
--- /dev/null
+++ b/Engine/Plugins/Enterprise/DatasmithImporter/Source/DatasmithImporter/Public/MasterMaterials/DatasmithMasterMaterialManager.h
@@ -0,0 +1,28 @@
+// Copyright 1998-2019 Epic Games, Inc. All Rights Reserved.
+#pragma once
+
+#include "Containers/ContainersFwd.h"
+#include "Templates/SharedPointer.h"
+#include "Templates/UniquePtr.h"
+
+class FDatasmithMasterMaterialSelector;
+
+class DATASMITHIMPORTER_API FDatasmithMasterMaterialManager
+{
+public:
+	static void Create();
+	static void Destroy();
+	static FDatasmithMasterMaterialManager& Get();
+
+	FString GetHostFromString( const FString& HostString );
+
+	void RegisterSelector( const TCHAR* Host, TSharedPtr< FDatasmithMasterMaterialSelector > Selector );
+	void UnregisterSelector( const TCHAR* Host );
+
+	const TSharedPtr< FDatasmithMasterMaterialSelector > GetSelector( const TCHAR* Host ) const;
+
+private:
+	static TUniquePtr< FDatasmithMasterMaterialManager > Instance;
+
+	TMap< FString, TSharedPtr< FDatasmithMasterMaterialSelector > > Selectors;
+};
\ No newline at end of file
diff --git a/Engine/Plugins/Enterprise/DatasmithImporter/Source/DatasmithImporter/Public/MasterMaterials/DatasmithMasterMaterialSelector.h b/Engine/Plugins/Enterprise/DatasmithImporter/Source/DatasmithImporter/Public/MasterMaterials/DatasmithMasterMaterialSelector.h
new file mode 100644
--- /dev/null
+++ b/Engine/Plugins/Enterprise/DatasmithImporter/Source/DatasmithImporter/Public/MasterMaterials/DatasmithMasterMaterialSelector.h
@@ -0,0 +1,31 @@
+// Copyright 1998-2019 Epic Games, Inc. All Rights Reserved.
+#pragma once
+
+#include "Math/Color.h"
+#include "Templates/SharedPointer.h"
+
+class FDatasmithMasterMaterial;
+class IDatasmithMasterMaterialElement;
+class IDatasmithKeyValueProperty;
+class UMaterialInstanceConstant;
+
+class DATASMITHIMPORTER_API FDatasmithMasterMaterialSelector
+{
+public:
+	virtual ~FDatasmithMasterMaterialSelector() = default;
+
+	virtual bool IsValid() const { return false; }
+	virtual const FDatasmithMasterMaterial& GetMasterMaterial( const TSharedPtr< IDatasmithMasterMaterialElement >& InDatasmithMaterial ) const;
+	virtual void FinalizeMaterialInstance(const TSharedPtr< IDatasmithMasterMaterialElement >& InDatasmithMaterial, UMaterialInstanceConstant* MaterialInstance) const {}
+
+	virtual bool GetColor( const TSharedPtr< IDatasmithKeyValueProperty >& InMaterialProperty, FLinearColor& OutValue ) const;
+	virtual bool GetFloat( const TSharedPtr< IDatasmithKeyValueProperty >& InMaterialProperty, float& OutValue ) const;
+	virtual bool GetBool( const TSharedPtr< IDatasmithKeyValueProperty >& InMaterialProperty, bool& OutValue ) const;
+	virtual bool GetTexture( const TSharedPtr< IDatasmithKeyValueProperty >& InMaterialProperty, FString& OutValue ) const;
+	virtual bool GetString( const TSharedPtr< IDatasmithKeyValueProperty >& InMaterialProperty, FString& OutValue ) const;
+
+protected:
+	static FDatasmithMasterMaterial InvalidMasterMaterial;
+
+};
+
diff --git a/Engine/Plugins/Enterprise/DatasmithImporter/Source/DatasmithImporter/Public/UI/DatasmithUICommands.h b/Engine/Plugins/Enterprise/DatasmithImporter/Source/DatasmithImporter/Public/UI/DatasmithUICommands.h
new file mode 100644
--- /dev/null
+++ b/Engine/Plugins/Enterprise/DatasmithImporter/Source/DatasmithImporter/Public/UI/DatasmithUICommands.h
@@ -0,0 +1,24 @@
+// Copyright 1998-2019 Epic Games, Inc. All Rights Reserved.
+
+#pragma once
+#include "Framework/Commands/Commands.h"
+
+class FDatasmithUICommands : public TCommands<FDatasmithUICommands>
+{
+public:
+	FDatasmithUICommands();
+
+	// TCommands interface
+	virtual void RegisterCommands() override;
+	// End of TCommands interface
+
+	// Add a new command to the Datasmith menu
+	static TSharedPtr< FUICommandInfo > AddMenuCommand(const FString& CommandName, const FText& Caption, const FText& Description);
+
+	// Remove a command from the Datasmith menu
+	static void RemoveMenuCommand(const TSharedPtr< FUICommandInfo >& Command);
+
+public:
+	TSharedPtr<FUICommandInfo> RepeatLastImport;
+	TArray<TSharedPtr<FUICommandInfo>> MenuCommands;
+};
diff --git a/Engine/Plugins/Enterprise/DatasmithImporter/Source/DatasmithImporter/Public/UI/DatasmithUIManager.h b/Engine/Plugins/Enterprise/DatasmithImporter/Source/DatasmithImporter/Public/UI/DatasmithUIManager.h
new file mode 100644
--- /dev/null
+++ b/Engine/Plugins/Enterprise/DatasmithImporter/Source/DatasmithImporter/Public/UI/DatasmithUIManager.h
@@ -0,0 +1,57 @@
+// Copyright 1998-2019 Epic Games, Inc. All Rights Reserved.
+
+#pragma once
+#include "Containers/Map.h"
+#include "CoreMinimal.h"
+#include "Framework/Commands/UIAction.h"
+
+class FUICommandInfo;
+class FUICommandList;
+class FToolBarBuilder;
+
+class DATASMITHIMPORTER_API FDatasmithUIManager
+{
+public:
+	static void Initialize();
+	static void Shutdown();
+
+	/** Get the Datasmith UI manager singleton instance. */
+	static FDatasmithUIManager& Get();
+
+	/**
+	 * Add a menu entry to the Datasmith importers drop-down menu
+	 *
+	 * @param CommandName		The command name (for internal references)
+	 * @param Caption			The text displayed in the menu entry
+	 * @param Description		The tooltip text displayed when hovering over the menu entry
+	 * @param IconResourcePath	The relative path to the .png file for the icon to be displayed in the menu
+	 * @param ExecuteAction		The action to be executed when the menu entry is selected
+	 * @param FactoryClass		The UClass of the factory to associate with the menu entry
+	 *
+	 * @return the FUICommandInfo that was created for that menu entry
+	 */
+	TSharedPtr<FUICommandInfo> AddMenuEntry(const FString& CommandName, const FText& Caption, const FText& Description, const FString& IconResourcePath, FExecuteAction ExecuteAction, UClass* FactoryClass);
+
+	/**
+	 * Remove a menu entry from the Datasmith importers drop-down menu
+	 *
+	 * @param Command			The FUICommandInfo to remove
+	 */
+	void RemoveMenuEntry(const TSharedPtr<FUICommandInfo>& Command);
+
+	/** Set the last factory used through the menu */
+	void SetLastFactoryUsed(UClass* Class);
+
+private:
+	static TUniquePtr<FDatasmithUIManager> Instance;
+
+	TSharedPtr<FUICommandList> DatasmithActions;
+	TSharedPtr<FUICommandInfo> LastSelectedCommand;
+
+	TMap<UClass*, TSharedPtr<FUICommandInfo>> FactoryClassToUICommandMap;
+
+	void ExtendToolbar();
+	void FillToolbar(FToolBarBuilder& ToolbarBuilder);
+
+	const TSharedRef<FUICommandInfo> GetLastSelectedCommand();
+};
\ No newline at end of file
diff --git a/Engine/Plugins/Enterprise/DatasmithImporter/Source/DatasmithImporter/Public/Utility/DatasmithMathUtils.h b/Engine/Plugins/Enterprise/DatasmithImporter/Source/DatasmithImporter/Public/Utility/DatasmithMathUtils.h
new file mode 100644
--- /dev/null
+++ b/Engine/Plugins/Enterprise/DatasmithImporter/Source/DatasmithImporter/Public/Utility/DatasmithMathUtils.h
@@ -0,0 +1,18 @@
+// Copyright 1998-2019 Epic Games, Inc. All Rights Reserved.
+#pragma once
+
+#include "CoreMinimal.h"
+
+
+struct FTransform;
+struct FQuat;
+
+class DATASMITHIMPORTER_API FDatasmithTransformUtils
+{
+public:
+
+	/**
+	 * Get rotation from FTransform and fix floating point precision issues with quaternion
+	 */
+	static void GetRotation(const FTransform& InTransform, FQuat& OutRotation);
+};
\ No newline at end of file
diff --git a/Engine/Plugins/Enterprise/DatasmithImporter/Source/DatasmithNativeTranslator/Private/DatasmithNativeTranslator.cpp b/Engine/Plugins/Enterprise/DatasmithImporter/Source/DatasmithNativeTranslator/Private/DatasmithNativeTranslator.cpp
new file mode 100644
--- /dev/null
+++ b/Engine/Plugins/Enterprise/DatasmithImporter/Source/DatasmithNativeTranslator/Private/DatasmithNativeTranslator.cpp
@@ -0,0 +1,170 @@
+// Copyright 1998-2019 Epic Games, Inc. All Rights Reserved.
+
+#include "DatasmithNativeTranslator.h"
+
+#include "DatasmithAnimationElements.h"
+#include "DatasmithAnimationSerializer.h"
+#include "DatasmithMeshHelper.h"
+#include "DatasmithMeshUObject.h"
+#include "DatasmithSceneSource.h"
+#include "DatasmithSceneXmlReader.h"
+#include "IDatasmithSceneElements.h"
+
+#include "HAL/FileManager.h"
+#include "MeshDescriptionOperations.h"
+#include "Serialization/MemoryReader.h"
+#include "Templates/SharedPointer.h"
+#include "Templates/UniquePtr.h"
+#include "UObject/StrongObjectPtr.h"
+
+
+
+void FDatasmithNativeTranslator::Initialize(FDatasmithTranslatorCapabilities& OutCapabilities)
+{
+	OutCapabilities.SupportedFileFormats.Emplace(TEXT("udatasmith"), TEXT("Datasmith files"));
+	OutCapabilities.bParallelLoadStaticMeshSupported = true;
+}
+
+bool FDatasmithNativeTranslator::LoadScene(TSharedRef<IDatasmithScene> OutScene)
+{
+	FDatasmithSceneXmlReader XmlParser;
+	return XmlParser.ParseFile(GetSource().GetSourceFile(), OutScene);
+}
+
+
+namespace DatasmithNativeTranslatorImpl
+{
+
+	TArray< UDatasmithMesh* > GetDatasmithMeshFromMeshElement( const TSharedRef< IDatasmithMeshElement > MeshElement )
+	{
+		TArray< UDatasmithMesh* > Result;
+
+		TUniquePtr<FArchive> Archive( IFileManager::Get().CreateFileReader(MeshElement->GetFile()) );
+		if ( !Archive.IsValid() )
+		{
+			return Result;
+		}
+
+		int32 NumMeshes;
+		*Archive << NumMeshes;
+
+		// Currently we only have 1 mesh per file. If there's a second mesh, it will be a CollisionMesh
+		while ( NumMeshes-- > 0)
+		{
+			TArray< uint8 > Bytes;
+			*Archive << Bytes;
+
+			UDatasmithMesh* DatasmithMesh = NewObject< UDatasmithMesh >();
+
+			FMemoryReader MemoryReader( Bytes, true );
+			MemoryReader.ArIgnoreClassRef = false;
+			MemoryReader.ArIgnoreArchetypeRef = false;
+			MemoryReader.SetWantBinaryPropertySerialization(true);
+			DatasmithMesh->Serialize( MemoryReader );
+
+			Result.Emplace( DatasmithMesh );
+		}
+
+		return Result;
+	}
+
+	TOptional<FMeshDescription> ExtractMeshDescription(FDatasmithMeshSourceModel& DSSourceModel)
+	{
+		FRawMesh RawMesh;
+		DSSourceModel.RawMeshBulkData.LoadRawMesh( RawMesh );
+		if ( !RawMesh.IsValid() )
+		{
+			return {};
+		}
+
+		// RawMesh -> MeshDescription conversion requires an {mat_index: slot_name} map for its PolygonGroups.
+		TMap<int32, FName> GroupNamePerGroupIndex;
+
+		// There is no guaranty that incoming RawMesh.FaceMaterialIndices are sequential, but the conversion assumes so.
+		// -> we remap materials identifiers to material indices
+		// eg:
+		//   incoming per-face mat identifier   5   5   1   1   1   99   99
+		//   remapped per-face index            0   0   1   1   1   2    2
+		//   per PolygonGroup FName:           "5" "5" "1" "1" "1" "99" "99"
+		TSet<int32> MaterialIdentifiers;
+		for (int32& MatIdentifier : RawMesh.FaceMaterialIndices)
+		{
+			bool bAlreadyIn = false;
+			int32 IndexOfIdentifier = MaterialIdentifiers.Add(MatIdentifier, &bAlreadyIn).AsInteger();
+
+			// identifier -> name association
+			if (!bAlreadyIn)
+			{
+				FName MaterialSlotName = DatasmithMeshHelper::DefaultSlotName(MatIdentifier);
+				GroupNamePerGroupIndex.Add(IndexOfIdentifier, MaterialSlotName);
+			}
+
+			// remap old identifier to material index
+			MatIdentifier = IndexOfIdentifier;
+		}
+
+		FMeshDescription MeshDescription;
+		DatasmithMeshHelper::PrepareAttributeForStaticMesh(MeshDescription);
+		FMeshDescriptionOperations::ConvertFromRawMesh(RawMesh, MeshDescription, GroupNamePerGroupIndex);
+		return MeshDescription;
+	}
+
+} // ns DatasmithNativeTranslatorImpl
+
+bool FDatasmithNativeTranslator::LoadStaticMesh(const TSharedRef<IDatasmithMeshElement> MeshElement, FDatasmithMeshElementPayload& OutMeshPayload)
+{
+	TRACE_CPUPROFILER_EVENT_SCOPE(FDatasmithNativeTranslator::LoadStaticMesh);
+
+	using namespace DatasmithNativeTranslatorImpl;
+
+	int32 ExtractionFailure = 0;
+	for (UDatasmithMesh* DatasmithMesh : GetDatasmithMeshFromMeshElement( MeshElement ))
+	{
+		if (DatasmithMesh->bIsCollisionMesh)
+		{
+			for (FDatasmithMeshSourceModel& SourceModel : DatasmithMesh->SourceModels)
+			{
+				FRawMesh RawMesh;
+				SourceModel.RawMeshBulkData.LoadRawMesh( RawMesh );
+				if ( RawMesh.VertexPositions.Num() > 0 )
+				{
+					OutMeshPayload.CollisionPointCloud = MoveTemp(RawMesh.VertexPositions);
+					break;
+				}
+				ExtractionFailure++;
+			}
+		}
+		else
+		{
+			for (FDatasmithMeshSourceModel& SourceModel : DatasmithMesh->SourceModels)
+			{
+				if (TOptional<FMeshDescription> Mesh = ExtractMeshDescription(SourceModel))
+				{
+					OutMeshPayload.LodMeshes.Add(MoveTemp(Mesh.GetValue()));
+					continue;
+				}
+				ExtractionFailure++;
+			}
+		}
+
+		// Do not wait until garbage collection to free memory of the models
+		DatasmithMesh->SourceModels.Reset();
+		
+		// Object has been created on another thread, tell the garbage collector it can now be deleted.
+		DatasmithMesh->ClearInternalFlags(EInternalObjectFlags::Async);
+	}
+
+	return ExtractionFailure == 0;
+}
+
+bool FDatasmithNativeTranslator::LoadLevelSequence(const TSharedRef<IDatasmithLevelSequenceElement> LevelSequenceElement, FDatasmithLevelSequencePayload& OutLevelSequencePayload)
+{
+	// #ueent_todo: this totally skip the payload system....
+	// Parse the level sequences from file
+	FDatasmithAnimationSerializer AnimSerializer;
+	if (LevelSequenceElement->GetFile() && IFileManager::Get().FileExists(LevelSequenceElement->GetFile()))
+	{
+		AnimSerializer.Deserialize(LevelSequenceElement, LevelSequenceElement->GetFile());
+	}
+	return true;
+}
