Datasmith
- Add missing call to UpdateCachedExpressions for newly created materials through Datasmith as we'll be creating instances before we finalize the material itself.
- Don't limit Datasmith textures to 4k

diff --git a/Engine/Plugins/Enterprise/DatasmithImporter/Source/DatasmithImporter/Private/DatasmithMaterialExpressions.cpp b/Engine/Plugins/Enterprise/DatasmithImporter/Source/DatasmithImporter/Private/DatasmithMaterialExpressions.cpp
--- a/Engine/Plugins/Enterprise/DatasmithImporter/Source/DatasmithImporter/Private/DatasmithMaterialExpressions.cpp
+++ b/Engine/Plugins/Enterprise/DatasmithImporter/Source/DatasmithImporter/Private/DatasmithMaterialExpressions.cpp
@@ -3192,6 +3192,8 @@ UMaterialInterface* FDatasmithMaterialExpressions::CreateUEPbrMaterial(UPackage*
 		UnrealMaterial->D3D11TessellationMode = EMaterialTessellationMode::MTM_FlatTessellation;
 	}
 
+	UnrealMaterial->UpdateCachedExpressionData();
+
 	UMaterialEditingLibrary::LayoutMaterialExpressions( UnrealMaterial );
 
 	return UnrealMaterial;
diff --git a/Engine/Plugins/Enterprise/DatasmithImporter/Source/DatasmithImporter/Private/DatasmithTextureImporter.cpp b/Engine/Plugins/Enterprise/DatasmithImporter/Source/DatasmithImporter/Private/DatasmithTextureImporter.cpp
--- a/Engine/Plugins/Enterprise/DatasmithImporter/Source/DatasmithImporter/Private/DatasmithTextureImporter.cpp
+++ b/Engine/Plugins/Enterprise/DatasmithImporter/Source/DatasmithImporter/Private/DatasmithTextureImporter.cpp
@@ -31,18 +31,17 @@
 #include "Modules/ModuleManager.h"
 
 #include "ObjectTools.h"
+#include "RHI.h"
 
 #define LOCTEXT_NAMESPACE "DatasmithTextureImport"
 
-const uint32 MaxTextureSize = 4096;
-
 namespace
 {
 	static bool ResizeTexture(const TCHAR* Filename, const TCHAR* ResizedFilename, bool bCreateNormal, FDatasmithImportContext& ImportContext)
 	{
 		TRACE_CPUPROFILER_EVENT_SCOPE(ResizeTexture);
 
-		EDSTextureUtilsError ErrorCode = FDatasmithTextureResize::ResizeTexture(Filename, ResizedFilename, EDSResizeTextureMode::NearestPowerOfTwo, MaxTextureSize, bCreateNormal);
+		EDSTextureUtilsError ErrorCode = FDatasmithTextureResize::ResizeTexture(Filename, ResizedFilename, EDSResizeTextureMode::NearestPowerOfTwo, GMaxTextureDimensions, bCreateNormal);
 
 		switch (ErrorCode)
 		{
