Merging //UE4/Dev-Enterprise to Release-4.25 (//UE4/Release-4.25)

CL 11163720

Dataprep: Multi-files assembly management via PLMXML

diff --git a/Engine/Plugins/Enterprise/DatasmithImporter/Source/DatasmithImporter/Private/DatasmithFileProducer.cpp b/Engine/Plugins/Enterprise/DatasmithImporter/Source/DatasmithImporter/Private/DatasmithFileProducer.cpp
--- a/Engine/Plugins/Enterprise/DatasmithImporter/Source/DatasmithImporter/Private/DatasmithFileProducer.cpp
+++ b/Engine/Plugins/Enterprise/DatasmithImporter/Source/DatasmithImporter/Private/DatasmithFileProducer.cpp
@@ -40,6 +40,7 @@
 #include "Materials/MaterialInterface.h"
 #include "Materials/MaterialInstance.h"
 #include "Misc/ConfigCacheIni.h"
+#include "Misc/FileHelper.h"
 #include "ObjectTools.h"
 #include "PropertyHandle.h"
 #include "ScopedTransaction.h"
@@ -50,7 +51,6 @@
 #include "Widgets/SWidget.h"
 #include "Widgets/Text/STextBlock.h"
 
-
 #define LOCTEXT_NAMESPACE "DatasmithFileProducer"
 
 const FText DatasmithFileProducerLabel( LOCTEXT( "DatasmithFileProducerLabel", "Datasmith file importer" ) );
@@ -782,7 +782,13 @@ bool UDatasmithDirProducer::Execute(TArray< TWeakObjectPtr< UObject > >& OutAsse
 	UPackage* RootTransientPackage = NewObject< UPackage >( nullptr, *RootPath, RF_Transient );
 	RootTransientPackage->FullyLoad();
 
-	for( const FString& FileName : FilesToProcess )
+	TArray<FString> FilesNotProcessed; // these files will be processed one by one
+	if (!ImportAsPlmXml(RootTransientPackage, OutAssets, FilesNotProcessed))
+	{
+		FilesNotProcessed = FilesToProcess.Array();
+	}
+	
+	for( const FString& FileName : FilesNotProcessed )
 	{
 		if ( IsCancelled() )
 		{
@@ -820,6 +826,150 @@ bool UDatasmithDirProducer::Execute(TArray< TWeakObjectPtr< UObject > >& OutAsse
 	return !IsCancelled();
 }
 
+// #ueent_todo:  ImportAsPlmXml should be replaced by a new object which will be responsible for the dispatch of import files which are supported by Datasmith
+/** Try to import as many files as possible using PlmXml translator. This is a temporary solution which enables importing multiple files at once using as many CPU cores as possible(fastest).
+ *  DatasmithDispatcher is able to import multiple files in parallel and PlmXmlTranslator is making use of this property of the DatasmithDispatcher.
+ *  Returns false if no files were processed
+ */
+bool UDatasmithDirProducer::ImportAsPlmXml(UPackage* RootPackage, TArray<TWeakObjectPtr<UObject>>& OutAssets, TArray<FString>& FilesNotProcessed)
+{
+	// Find out which translator can import CAD files by retrieving translator for a JT file(this is expected to be 'DatasmithCADTranslator').
+	// And then accept files which have this translator returned as compatible.
+	FDatasmithSceneSource SomeCADFileSource;
+	SomeCADFileSource.SetSourceFile("test.jt");
+	TSharedPtr<IDatasmithTranslator> TranslatorForCADFiles = FDatasmithTranslatorManager::Get().SelectFirstCompatible(SomeCADFileSource);
+	if (!TranslatorForCADFiles.IsValid())
+	{
+		return false;
+	}
+
+	FName CADTranslatorName = TranslatorForCADFiles->GetFName();
+	TArray<FString> FilesToProcessWithPlmXml;
+	for (const FString& FileName : FilesToProcess)
+	{
+		FDatasmithSceneSource Source;
+		Source.SetSourceFile(FileName);
+		TSharedPtr<IDatasmithTranslator> Translator = FDatasmithTranslatorManager::Get().SelectFirstCompatible(Source);
+
+		bool bIsCADFile = Translator.IsValid() && (Translator->GetFName() == CADTranslatorName);
+		if (bIsCADFile)
+		{
+			FilesToProcessWithPlmXml.Add(FileName);
+		}
+		else
+		{
+			FilesNotProcessed.Add(FileName);
+		}
+	}
+
+	// skip plmxml processing for single file
+	if (FilesToProcessWithPlmXml.Num() <= 1)
+	{
+		return false;
+	}
+
+	FString TempDir = FPaths::Combine(FPaths::ProjectIntermediateDir(), TEXT("DatasmithProducerTemp"));
+	if (!IFileManager::Get().DirectoryExists(*TempDir))
+	{
+		IFileManager::Get().MakeDirectory(*TempDir);
+	}
+	FString PlmXmlFileName = FPaths::Combine(TempDir, FPaths::GetBaseFilename(FolderPath) + TEXT(".plmxml"));
+
+	class FXmlWriter
+	{
+	public:
+		FString Buffer;
+
+		FXmlWriter()
+		{
+			Buffer = TEXT("<?xml version=\"1.0\" encoding=\"utf-8\"?>");
+			Buffer += LINE_TERMINATOR;
+		}
+
+		class FTagGuard
+		{
+		public:
+			FTagGuard(FXmlWriter& InWriter, FString Opening, FString InClosing)
+				: Writer(InWriter)
+				, Closing(InClosing)
+			{
+				Writer.Buffer += Opening;
+				Writer.Buffer += LINE_TERMINATOR;
+			}
+			~FTagGuard()
+			{
+				Writer.Buffer += Closing;
+				Writer.Buffer += LINE_TERMINATOR;
+			}
+		private:
+			FXmlWriter& Writer;
+			FString Closing;
+		};
+	};
+
+	FXmlWriter Writer;
+	FString& Buffer = Writer.Buffer;
+
+	// Creating PLMXML file where each of files to process is referenced from a ProductRevisionView
+	{
+		FXmlWriter::FTagGuard PLMXMLTag(Writer, TEXT("<PLMXML xmlns=\"http://www.plmxml.org/Schemas/PLMXMLSchema\">"), "</PLMXML>");
+		FXmlWriter::FTagGuard ProductDefTag(Writer, TEXT("<ProductDef id=\"id1\">"), TEXT("</ProductDef>"));
+
+		// Used to assign unique ids to PLMXML entities being created
+		int32 CurrentId = 2;
+
+		// Collect all InstanceId to reference from InstanceGraph rootRefs
+		TArray<FString> InstanceIds;
+		InstanceIds.Reserve(FilesToProcessWithPlmXml.Num());
+		for (const FString& FileName : FilesToProcessWithPlmXml)
+		{
+			InstanceIds.Add(FString::Printf(TEXT("id%d"), CurrentId++));
+		}
+		FXmlWriter::FTagGuard InstanceGraphTag(Writer, FString::Printf(TEXT("<InstanceGraph id=\"id2\" rootRefs=\"%s\">"), *FString::Join(InstanceIds, TEXT(" "))), TEXT("</InstanceGraph>"));
+
+		for (int32 FileIndex = 0; FileIndex < FilesToProcessWithPlmXml.Num(); ++FileIndex)
+		{
+			const FString& FileName = FilesToProcessWithPlmXml[FileIndex];
+			FString InstanceId = InstanceIds[FileIndex];
+			FString InstanceName = FPaths::GetBaseFilename(FileName);
+			FString PartId = FString::Printf(TEXT("id%d"), CurrentId++);
+			FString RepresentationId = FString::Printf(TEXT("id%d"), CurrentId++);
+			{
+				FXmlWriter::FTagGuard ProductInstanceTag(Writer, FString::Printf(TEXT("<ProductInstance id=\"%s\" name=\"%s\" partRef=\"#%s\">"), *InstanceId, *InstanceName, *PartId), TEXT("</ProductInstance>"));
+			}
+			{
+				FXmlWriter::FTagGuard ProductRevisionViewTag(Writer, FString::Printf(TEXT("<ProductRevisionView id=\"%s\" name=\"%s\">"), *PartId, *InstanceName), TEXT("</ProductRevisionView>"));
+				//  omitting 'format' attribute, it's optional anyway
+				FXmlWriter::FTagGuard RepresentationTag(Writer, FString::Printf(TEXT("<Representation id=\"%s\" location=\"%s\">"), *RepresentationId, *FileName), TEXT("</Representation>"));
+			}
+		}
+	}
+
+	if (!FFileHelper::SaveStringToFile(Buffer, *PlmXmlFileName, FFileHelper::EEncodingOptions::ForceUTF8WithoutBOM))
+	{
+		FText ErrorReport = FText::Format(LOCTEXT("DatasmithDirProducer_Failed", "Failed to create PlmXml file, {0}, for parallel loading ..."), FText::FromString(PlmXmlFileName));
+		LogError(ErrorReport);
+		return false;
+	}
+
+	// Import content of file into the proper content folder to avoid name collision
+	UPackage* TransientPackage = RootPackage;
+
+	Context.SetRootPackage(TransientPackage);
+
+	// Update file producer's filename
+	FileProducer->FilePath = PlmXmlFileName;
+	FileProducer->UpdateName();
+
+	if (!FileProducer->Produce(Context, OutAssets))
+	{
+		FText ErrorReport = FText::Format(LOCTEXT("DatasmithDirProducer_Failed", "Failed to produce assets with PlmXml file, {0}, for parallel loading ..."), FText::FromString(PlmXmlFileName));
+		LogError(ErrorReport);
+		return false;
+	}
+	return true;
+}
+
 void UDatasmithDirProducer::Reset()
 {
 	FilesToProcess.Empty();
diff --git a/Engine/Plugins/Enterprise/DatasmithImporter/Source/DatasmithImporter/Private/DatasmithFileProducer.h b/Engine/Plugins/Enterprise/DatasmithImporter/Source/DatasmithImporter/Private/DatasmithFileProducer.h
--- a/Engine/Plugins/Enterprise/DatasmithImporter/Source/DatasmithImporter/Private/DatasmithFileProducer.h
+++ b/Engine/Plugins/Enterprise/DatasmithImporter/Source/DatasmithImporter/Private/DatasmithFileProducer.h
@@ -175,6 +175,9 @@ private:
 	/** Update the name of the producer based on the directory name */
 	void UpdateName();
 
+	/** Try import files in the FolderPath as parts of constructed PLMXML file */
+	bool ImportAsPlmXml(UPackage* RootPackage, TArray< TWeakObjectPtr< UObject > >& OutAssets, TArray<FString>& FilesCouldntBeProcessedWithPlmXml);
+
 	/** Indicates if ExtensionString contains "*.*" */
 	bool bHasWildCardSearch;
 
