Incorrect names on Subcomponents imported with Datasmith

Integrated from Dev-Enterprise (CL 11153373)

diff --git a/Engine/Plugins/Enterprise/DatasmithImporter/Source/DatasmithImporter/Private/DatasmithActorImporter.cpp b/Engine/Plugins/Enterprise/DatasmithImporter/Source/DatasmithImporter/Private/DatasmithActorImporter.cpp
--- a/Engine/Plugins/Enterprise/DatasmithImporter/Source/DatasmithImporter/Private/DatasmithActorImporter.cpp
+++ b/Engine/Plugins/Enterprise/DatasmithImporter/Source/DatasmithImporter/Private/DatasmithActorImporter.cpp
@@ -136,7 +136,7 @@ AActor* FDatasmithActorImporter::ImportActor( UClass* ActorClass, const TSharedR
 	return ImportedActor;
 }
 
-USceneComponent* FDatasmithActorImporter::ImportSceneComponent( UClass* ComponentClass, const TSharedRef< IDatasmithActorElement >& ActorElement, FDatasmithImportContext& ImportContext, UObject* Outer )
+USceneComponent* FDatasmithActorImporter::ImportSceneComponent( UClass* ComponentClass, const TSharedRef< IDatasmithActorElement >& ActorElement, FDatasmithImportContext& ImportContext, UObject* Outer, FDatasmithActorUniqueLabelProvider& UniqueNameProvider )
 {
 	if ( !ComponentClass->IsChildOf( USceneComponent::StaticClass() ) )
 	{
@@ -187,15 +187,14 @@ USceneComponent* FDatasmithActorImporter::ImportSceneComponent( UClass* Componen
 
 	if ( !ValidSceneComponent )
 	{
-		if ( SceneComponent || FindObjectWithOuter( Outer, UObject::StaticClass() , ActorElement->GetLabel() ) )
+		FName ComponentName( ActorElement->GetLabel() );
+		if ( SceneComponent || FindObjectWithOuter( Outer, UObject::StaticClass(), ComponentName ) )
 		{
-			// There is already a object with this name inside the outer. Let the new object assigning it a name.
-			ValidSceneComponent = NewObject< USceneComponent >( Outer, ComponentClass, NAME_None, RF_Transactional );
-		}
-		else
-		{
-			ValidSceneComponent = NewObject< USceneComponent >( Outer, ComponentClass, ActorElement->GetLabel(), RF_Transactional );
+			// There is already a object with this name inside the outer. Generate unique name.
+			UniqueNameProvider.AddExistingName( ComponentName.ToString() );
+			ComponentName = *UniqueNameProvider.GenerateUniqueName( ComponentName.ToString() );
 		}
+		ValidSceneComponent = NewObject< USceneComponent >( Outer, ComponentClass, ComponentName, RF_Transactional );
 	}
 
 	if ( !ValidSceneComponent )
@@ -233,14 +232,14 @@ AActor* FDatasmithActorImporter::ImportBaseActor( FDatasmithImportContext& Impor
 	return Actor;
 }
 
-USceneComponent* FDatasmithActorImporter::ImportBaseActorAsComponent(FDatasmithImportContext& ImportContext, const TSharedRef< IDatasmithActorElement >& ActorElement, UObject* Outer)
+USceneComponent* FDatasmithActorImporter::ImportBaseActorAsComponent(FDatasmithImportContext& ImportContext, const TSharedRef< IDatasmithActorElement >& ActorElement, UObject* Outer, FDatasmithActorUniqueLabelProvider& UniqueNameProvider)
 {
 	if ( ImportContext.Options->OtherActorImportPolicy == EDatasmithImportActorPolicy::Ignore )
 	{
 		return nullptr;
 	}
 
-	USceneComponent* SceneComponent = FDatasmithActorImporter::ImportSceneComponent( USceneComponent::StaticClass(), ActorElement, ImportContext, Outer );
+	USceneComponent* SceneComponent = FDatasmithActorImporter::ImportSceneComponent( USceneComponent::StaticClass(), ActorElement, ImportContext, Outer, UniqueNameProvider );
 	SceneComponent->RegisterComponent();
 
 	ImportContext.AddSceneComponent(SceneComponent->GetName(), SceneComponent);
@@ -264,14 +263,14 @@ AStaticMeshActor* FDatasmithActorImporter::ImportStaticMeshActor(FDatasmithImpor
 	return StaticMeshActor;
 }
 
-UStaticMeshComponent* FDatasmithActorImporter::ImportStaticMeshComponent( FDatasmithImportContext& ImportContext, const TSharedRef<IDatasmithMeshActorElement>& InMeshActor, UObject* Outer )
+UStaticMeshComponent* FDatasmithActorImporter::ImportStaticMeshComponent( FDatasmithImportContext& ImportContext, const TSharedRef<IDatasmithMeshActorElement>& InMeshActor, UObject* Outer, FDatasmithActorUniqueLabelProvider& UniqueNameProvider )
 {
 	if ( ImportContext.Options->StaticMeshActorImportPolicy == EDatasmithImportActorPolicy::Ignore )
 	{
 		return nullptr;
 	}
 
-	USceneComponent* SceneComponent =  FDatasmithActorImporter::ImportSceneComponent( UStaticMeshComponent::StaticClass(), InMeshActor, ImportContext, Outer );
+	USceneComponent* SceneComponent =  FDatasmithActorImporter::ImportSceneComponent( UStaticMeshComponent::StaticClass(), InMeshActor, ImportContext, Outer, UniqueNameProvider );
 	UStaticMeshComponent* StaticMeshComponent = Cast< UStaticMeshComponent >( SceneComponent );
 
 	SetupStaticMeshComponent( ImportContext, StaticMeshComponent, InMeshActor );
@@ -324,7 +323,7 @@ AActor* FDatasmithActorImporter::ImportCameraActor(FDatasmithImportContext& Impo
 	return FDatasmithCameraImporter::ImportCameraActor(InCameraActor, ImportContext);
 }
 
-AActor* FDatasmithActorImporter::ImportCustomActor(FDatasmithImportContext& ImportContext, const TSharedRef< IDatasmithCustomActorElement >& InCustomActorElement)
+AActor* FDatasmithActorImporter::ImportCustomActor(FDatasmithImportContext& ImportContext, const TSharedRef< IDatasmithCustomActorElement >& InCustomActorElement, FDatasmithActorUniqueLabelProvider& UniqueNameProvider)
 {
 	if ( ImportContext.Options->OtherActorImportPolicy == EDatasmithImportActorPolicy::Ignore )
 	{
@@ -371,7 +370,7 @@ AActor* FDatasmithActorImporter::ImportCustomActor(FDatasmithImportContext& Impo
 
 	if ( !Actor->GetRootComponent() )
 	{
-		USceneComponent* RootComponent = FDatasmithActorImporter::ImportSceneComponent( USceneComponent::StaticClass(), InCustomActorElement, ImportContext, Actor );
+		USceneComponent* RootComponent = FDatasmithActorImporter::ImportSceneComponent( USceneComponent::StaticClass(), InCustomActorElement, ImportContext, Actor, UniqueNameProvider );
 		RootComponent->bVisualizeComponent = true;
 
 		Actor->SetRootComponent( RootComponent );
@@ -495,7 +494,7 @@ AActor* FDatasmithActorImporter::ImportEnvironment(FDatasmithImportContext& Impo
 	return Actor;
 }
 
-AActor* FDatasmithActorImporter::ImportHierarchicalInstancedStaticMeshAsActor(FDatasmithImportContext& ImportContext, const TSharedRef< IDatasmithHierarchicalInstancedStaticMeshActorElement >& HierarchicalInstancedStaticMeshActorElement)
+AActor* FDatasmithActorImporter::ImportHierarchicalInstancedStaticMeshAsActor(FDatasmithImportContext& ImportContext, const TSharedRef< IDatasmithHierarchicalInstancedStaticMeshActorElement >& HierarchicalInstancedStaticMeshActorElement, FDatasmithActorUniqueLabelProvider& UniqueNameProvider)
 {
 	AActor* Actor = ImportActor(AActor::StaticClass(), HierarchicalInstancedStaticMeshActorElement, ImportContext, ImportContext.Options->StaticMeshActorImportPolicy);
 
@@ -506,12 +505,12 @@ AActor* FDatasmithActorImporter::ImportHierarchicalInstancedStaticMeshAsActor(FD
 		if (Actor->GetRootComponent())
 		{
 			ImportContext.Hierarchy.Push(Actor->GetRootComponent());
-			ImportHierarchicalInstancedStaticMeshComponent(ImportContext, HierarchicalInstancedStaticMeshActorElement, Actor);
+			ImportHierarchicalInstancedStaticMeshComponent(ImportContext, HierarchicalInstancedStaticMeshActorElement, Actor, UniqueNameProvider);
 			ImportContext.Hierarchy.Pop();
 		}
 		else
 		{
-			ImportHierarchicalInstancedStaticMeshComponent(ImportContext, HierarchicalInstancedStaticMeshActorElement, Actor);
+			ImportHierarchicalInstancedStaticMeshComponent(ImportContext, HierarchicalInstancedStaticMeshActorElement, Actor, UniqueNameProvider);
 		}
 		HierarchicalInstancedStaticMeshActorElement->SetLabel(*OriginalLabel);
 	}
@@ -520,14 +519,14 @@ AActor* FDatasmithActorImporter::ImportHierarchicalInstancedStaticMeshAsActor(FD
 }
 
 UHierarchicalInstancedStaticMeshComponent* FDatasmithActorImporter::ImportHierarchicalInstancedStaticMeshComponent(FDatasmithImportContext& ImportContext, const TSharedRef< IDatasmithHierarchicalInstancedStaticMeshActorElement >& HierarchicalInstancedStatictMeshActorElement,
-	UObject* Outer)
+	UObject* Outer, FDatasmithActorUniqueLabelProvider& UniqueNameProvider)
 {
 	if ( ImportContext.Options->StaticMeshActorImportPolicy == EDatasmithImportActorPolicy::Ignore )
 	{
 		return nullptr;
 	}
 
-	USceneComponent* SeneComponent = FDatasmithActorImporter::ImportSceneComponent(UHierarchicalInstancedStaticMeshComponent::StaticClass(), HierarchicalInstancedStatictMeshActorElement, ImportContext, Outer);
+	USceneComponent* SeneComponent = FDatasmithActorImporter::ImportSceneComponent(UHierarchicalInstancedStaticMeshComponent::StaticClass(), HierarchicalInstancedStatictMeshActorElement, ImportContext, Outer, UniqueNameProvider);
 	UHierarchicalInstancedStaticMeshComponent* HierarchilcalInstanciedStaticMeshComponent = Cast< UHierarchicalInstancedStaticMeshComponent >(SeneComponent);
 
 	SetupHierarchicalInstancedStaticMeshComponent(ImportContext, HierarchilcalInstanciedStaticMeshComponent, HierarchicalInstancedStatictMeshActorElement);
diff --git a/Engine/Plugins/Enterprise/DatasmithImporter/Source/DatasmithImporter/Private/DatasmithCameraImporter.cpp b/Engine/Plugins/Enterprise/DatasmithImporter/Source/DatasmithImporter/Private/DatasmithCameraImporter.cpp
--- a/Engine/Plugins/Enterprise/DatasmithImporter/Source/DatasmithImporter/Private/DatasmithCameraImporter.cpp
+++ b/Engine/Plugins/Enterprise/DatasmithImporter/Source/DatasmithImporter/Private/DatasmithCameraImporter.cpp
@@ -39,9 +39,9 @@ AActor* FDatasmithCameraImporter::ImportCameraActor( const TSharedRef< IDatasmit
 	return CameraActor;
 }
 
-UCineCameraComponent* FDatasmithCameraImporter::ImportCineCameraComponent( const TSharedRef< IDatasmithCameraActorElement >& CameraElement, FDatasmithImportContext& ImportContext, UObject* Outer )
+UCineCameraComponent* FDatasmithCameraImporter::ImportCineCameraComponent( const TSharedRef< IDatasmithCameraActorElement >& CameraElement, FDatasmithImportContext& ImportContext, UObject* Outer, FDatasmithActorUniqueLabelProvider& UniqueNameProvider )
 {
-	UCineCameraComponent* CineCameraComponent = Cast< UCineCameraComponent >( FDatasmithActorImporter::ImportSceneComponent( UCineCameraComponent::StaticClass(), CameraElement, ImportContext, Outer ) );
+	UCineCameraComponent* CineCameraComponent = Cast< UCineCameraComponent >( FDatasmithActorImporter::ImportSceneComponent( UCineCameraComponent::StaticClass(), CameraElement, ImportContext, Outer, UniqueNameProvider ) );
 
 	if ( !CineCameraComponent )
 	{
diff --git a/Engine/Plugins/Enterprise/DatasmithImporter/Source/DatasmithImporter/Private/DatasmithImporter.cpp b/Engine/Plugins/Enterprise/DatasmithImporter/Source/DatasmithImporter/Private/DatasmithImporter.cpp
--- a/Engine/Plugins/Enterprise/DatasmithImporter/Source/DatasmithImporter/Private/DatasmithImporter.cpp
+++ b/Engine/Plugins/Enterprise/DatasmithImporter/Source/DatasmithImporter/Private/DatasmithImporter.cpp
@@ -1455,6 +1455,8 @@ void FDatasmithImporter::ImportActors( FDatasmithImportContext& ImportContext )
 	{
 		ImportContext.Hierarchy.Push( ImportSceneActor->GetRootComponent() );
 
+		FDatasmithActorUniqueLabelProvider UniqueNameProvider;
+
 		for (int32 i = 0; i < ActorsCount && !ImportContext.bUserCancelled; ++i)
 		{
 			ImportContext.bUserCancelled |= DatasmithImporterImpl::HasUserCancelledTask( ImportContext.FeedbackContext );
@@ -1467,7 +1469,7 @@ void FDatasmithImporter::ImportActors( FDatasmithImportContext& ImportContext )
 
 				if ( ActorElement->IsAComponent() )
 				{
-					ImportActorAsComponent( ImportContext, ActorElement.ToSharedRef(), ImportSceneActor );
+					ImportActorAsComponent( ImportContext, ActorElement.ToSharedRef(), ImportSceneActor, UniqueNameProvider );
 				}
 				else
 				{
@@ -1517,11 +1519,13 @@ AActor* FDatasmithImporter::ImportActor( FDatasmithImportContext& ImportContext,
 {
 	TRACE_CPUPROFILER_EVENT_SCOPE(FDatasmithImporter::ImportActor);
 
+	FDatasmithActorUniqueLabelProvider UniqueNameProvider;
+
 	AActor* ImportedActor = nullptr;
 	if (ActorElement->IsA(EDatasmithElementType::HierarchicalInstanceStaticMesh))
 	{
 		TSharedRef< IDatasmithHierarchicalInstancedStaticMeshActorElement > HISMActorElement = StaticCastSharedRef< IDatasmithHierarchicalInstancedStaticMeshActorElement >( ActorElement );
-		ImportedActor =  FDatasmithActorImporter::ImportHierarchicalInstancedStaticMeshAsActor( ImportContext, HISMActorElement );
+		ImportedActor =  FDatasmithActorImporter::ImportHierarchicalInstancedStaticMeshAsActor( ImportContext, HISMActorElement, UniqueNameProvider );
 	}
 	else if (ActorElement->IsA(EDatasmithElementType::StaticMeshActor))
 	{
@@ -1542,7 +1546,7 @@ AActor* FDatasmithImporter::ImportActor( FDatasmithImportContext& ImportContext,
 	}
 	else if (ActorElement->IsA(EDatasmithElementType::CustomActor))
 	{
-		ImportedActor = FDatasmithActorImporter::ImportCustomActor( ImportContext, StaticCastSharedRef< IDatasmithCustomActorElement >(ActorElement) );
+		ImportedActor = FDatasmithActorImporter::ImportCustomActor( ImportContext, StaticCastSharedRef< IDatasmithCustomActorElement >(ActorElement), UniqueNameProvider );
 	}
 	else if (ActorElement->IsA(EDatasmithElementType::Landscape))
 	{
@@ -1582,7 +1586,7 @@ AActor* FDatasmithImporter::ImportActor( FDatasmithImportContext& ImportContext,
 			}
 			else if ( ImportedActor ) // Don't import the components of an actor that we didn't import
 			{
-				ImportActorAsComponent( ImportContext, ChildActorElement.ToSharedRef(), ImportedActor );
+				ImportActorAsComponent( ImportContext, ChildActorElement.ToSharedRef(), ImportedActor, UniqueNameProvider );
 			}
 		}
 	}
@@ -1595,7 +1599,7 @@ AActor* FDatasmithImporter::ImportActor( FDatasmithImportContext& ImportContext,
 	return ImportedActor;
 }
 
-void FDatasmithImporter::ImportActorAsComponent(FDatasmithImportContext& ImportContext, const TSharedRef< IDatasmithActorElement >& ActorElement, AActor* InRootActor)
+void FDatasmithImporter::ImportActorAsComponent(FDatasmithImportContext& ImportContext, const TSharedRef< IDatasmithActorElement >& ActorElement, AActor* InRootActor, FDatasmithActorUniqueLabelProvider& UniqueNameProvider)
 {
 	if (!InRootActor)
 	{
@@ -1607,12 +1611,12 @@ void FDatasmithImporter::ImportActorAsComponent(FDatasmithImportContext& ImportC
 	if (ActorElement->IsA(EDatasmithElementType::HierarchicalInstanceStaticMesh))
 	{
 		TSharedRef< IDatasmithHierarchicalInstancedStaticMeshActorElement > HierarchicalInstancedStaticMeshElement = StaticCastSharedRef< IDatasmithHierarchicalInstancedStaticMeshActorElement >(ActorElement);
-		SceneComponent = FDatasmithActorImporter::ImportHierarchicalInstancedStaticMeshComponent(ImportContext, HierarchicalInstancedStaticMeshElement, InRootActor);
+		SceneComponent = FDatasmithActorImporter::ImportHierarchicalInstancedStaticMeshComponent(ImportContext, HierarchicalInstancedStaticMeshElement, InRootActor, UniqueNameProvider);
 	}
 	else if (ActorElement->IsA(EDatasmithElementType::StaticMeshActor))
 	{
 		TSharedRef< IDatasmithMeshActorElement > MeshActorElement = StaticCastSharedRef< IDatasmithMeshActorElement >(ActorElement);
-		SceneComponent = FDatasmithActorImporter::ImportStaticMeshComponent(ImportContext, MeshActorElement, InRootActor);
+		SceneComponent = FDatasmithActorImporter::ImportStaticMeshComponent(ImportContext, MeshActorElement, InRootActor, UniqueNameProvider);
 	}
 	else if (ActorElement->IsA(EDatasmithElementType::Light))
 	{
@@ -1621,7 +1625,7 @@ void FDatasmithImporter::ImportActorAsComponent(FDatasmithImportContext& ImportC
 			return;
 		}
 
-		SceneComponent = FDatasmithLightImporter::ImportLightComponent(StaticCastSharedRef< IDatasmithLightActorElement >(ActorElement), ImportContext, InRootActor);
+		SceneComponent = FDatasmithLightImporter::ImportLightComponent(StaticCastSharedRef< IDatasmithLightActorElement >(ActorElement), ImportContext, InRootActor, UniqueNameProvider);
 	}
 	else if (ActorElement->IsA(EDatasmithElementType::Camera))
 	{
@@ -1630,11 +1634,11 @@ void FDatasmithImporter::ImportActorAsComponent(FDatasmithImportContext& ImportC
 			return;
 		}
 
-		SceneComponent = FDatasmithCameraImporter::ImportCineCameraComponent(StaticCastSharedRef< IDatasmithCameraActorElement >(ActorElement), ImportContext, InRootActor);
+		SceneComponent = FDatasmithCameraImporter::ImportCineCameraComponent(StaticCastSharedRef< IDatasmithCameraActorElement >(ActorElement), ImportContext, InRootActor, UniqueNameProvider);
 	}
 	else
 	{
-		SceneComponent = FDatasmithActorImporter::ImportBaseActorAsComponent(ImportContext, ActorElement, InRootActor);
+		SceneComponent = FDatasmithActorImporter::ImportBaseActorAsComponent(ImportContext, ActorElement, InRootActor, UniqueNameProvider);
 	}
 
 	if (SceneComponent)
@@ -1654,7 +1658,7 @@ void FDatasmithImporter::ImportActorAsComponent(FDatasmithImportContext& ImportC
 			ImportContext.Hierarchy.Push(SceneComponent);
 		}
 
-		ImportActorAsComponent(ImportContext, ActorElement->GetChild(i).ToSharedRef(), InRootActor);
+		ImportActorAsComponent(ImportContext, ActorElement->GetChild(i).ToSharedRef(), InRootActor, UniqueNameProvider);
 
 		if (SceneComponent)
 		{
diff --git a/Engine/Plugins/Enterprise/DatasmithImporter/Source/DatasmithImporter/Private/DatasmithLightImporter.cpp b/Engine/Plugins/Enterprise/DatasmithImporter/Source/DatasmithImporter/Private/DatasmithLightImporter.cpp
--- a/Engine/Plugins/Enterprise/DatasmithImporter/Source/DatasmithImporter/Private/DatasmithLightImporter.cpp
+++ b/Engine/Plugins/Enterprise/DatasmithImporter/Source/DatasmithImporter/Private/DatasmithLightImporter.cpp
@@ -204,7 +204,7 @@ AActor* FDatasmithLightImporter::ImportLightActor( const TSharedRef< IDatasmithL
 	return ImportedLightActor;
 }
 
-USceneComponent* FDatasmithLightImporter::ImportLightComponent( const TSharedRef< IDatasmithLightActorElement >& LightElement, FDatasmithImportContext& ImportContext, UObject* Outer )
+USceneComponent* FDatasmithLightImporter::ImportLightComponent( const TSharedRef< IDatasmithLightActorElement >& LightElement, FDatasmithImportContext& ImportContext, UObject* Outer, FDatasmithActorUniqueLabelProvider& UniqueNameProvider )
 {
 	CreateIESTexture( ImportContext, LightElement );
 
@@ -213,16 +213,16 @@ USceneComponent* FDatasmithLightImporter::ImportLightComponent( const TSharedRef
 	if ( LightElement->IsA(EDatasmithElementType::AreaLight) )
 	{
 		TSharedRef< IDatasmithAreaLightElement > AreaLightElement = StaticCastSharedRef< IDatasmithAreaLightElement >( LightElement );
-		LightComponent = ImportAreaLightComponent( AreaLightElement, ImportContext, Outer );
+		LightComponent = ImportAreaLightComponent( AreaLightElement, ImportContext, Outer, UniqueNameProvider );
 	}
 	else if ( LightElement->IsA( EDatasmithElementType::LightmassPortal ) )
 	{
 		TSharedRef< IDatasmithLightmassPortalElement > LightmassPortalElement = StaticCastSharedRef< IDatasmithLightmassPortalElement >( LightElement );
-		LightComponent = ImportLightmassPortalComponent( LightmassPortalElement, ImportContext, Outer );
+		LightComponent = ImportLightmassPortalComponent( LightmassPortalElement, ImportContext, Outer, UniqueNameProvider );
 	}
 	else if ( LightElement->IsA( EDatasmithElementType::DirectionalLight ) )
 	{
-		USceneComponent* SceneComponent = FDatasmithActorImporter::ImportSceneComponent( UDirectionalLightComponent::StaticClass(), LightElement, ImportContext, Outer );
+		USceneComponent* SceneComponent = FDatasmithActorImporter::ImportSceneComponent( UDirectionalLightComponent::StaticClass(), LightElement, ImportContext, Outer, UniqueNameProvider );
 
 		UDirectionalLightComponent* DirectionalLightComponent = Cast< UDirectionalLightComponent >( SceneComponent );
 
@@ -247,7 +247,7 @@ USceneComponent* FDatasmithLightImporter::ImportLightComponent( const TSharedRef
 
 		if ( LightElement->IsA( EDatasmithElementType::SpotLight ) )
 		{
-			USceneComponent* SceneComponent = FDatasmithActorImporter::ImportSceneComponent( USpotLightComponent::StaticClass(), LightElement, ImportContext, Outer );
+			USceneComponent* SceneComponent = FDatasmithActorImporter::ImportSceneComponent( USpotLightComponent::StaticClass(), LightElement, ImportContext, Outer, UniqueNameProvider );
 
 			USpotLightComponent* SpotLightComponent = Cast< USpotLightComponent >( SceneComponent );
 
@@ -263,7 +263,7 @@ USceneComponent* FDatasmithLightImporter::ImportLightComponent( const TSharedRef
 		}
 		else if ( LightElement->IsA( EDatasmithElementType::PointLight ) )
 		{
-			USceneComponent* SceneComponent = FDatasmithActorImporter::ImportSceneComponent( UPointLightComponent::StaticClass(), LightElement, ImportContext, Outer );
+			USceneComponent* SceneComponent = FDatasmithActorImporter::ImportSceneComponent( UPointLightComponent::StaticClass(), LightElement, ImportContext, Outer, UniqueNameProvider );
 
 			UPointLightComponent* PointLightComponent = Cast< UPointLightComponent >( SceneComponent );
 
@@ -445,7 +445,7 @@ AActor* FDatasmithLightImporter::ImportAreaLightActor( const TSharedRef< IDatasm
 	return AreaLightActor;
 }
 
-USceneComponent* FDatasmithLightImporter::ImportAreaLightComponent( const TSharedRef< IDatasmithAreaLightElement >& AreaLightElement, FDatasmithImportContext& ImportContext, UObject* Outer )
+USceneComponent* FDatasmithLightImporter::ImportAreaLightComponent( const TSharedRef< IDatasmithAreaLightElement >& AreaLightElement, FDatasmithImportContext& ImportContext, UObject* Outer, FDatasmithActorUniqueLabelProvider& UniqueNameProvider )
 {
 	TRACE_CPUPROFILER_EVENT_SCOPE(FDatasmithLightImporter::ImportAreaLightComponent);
 
@@ -456,7 +456,7 @@ USceneComponent* FDatasmithLightImporter::ImportAreaLightComponent( const TShare
 
 	if ( LightShapeBlueprint )
 	{
-		UChildActorComponent* ChildActorComponent = Cast< UChildActorComponent >( FDatasmithActorImporter::ImportSceneComponent( UChildActorComponent::StaticClass(), AreaLightElement, ImportContext, Outer ) );
+		UChildActorComponent* ChildActorComponent = Cast< UChildActorComponent >( FDatasmithActorImporter::ImportSceneComponent( UChildActorComponent::StaticClass(), AreaLightElement, ImportContext, Outer, UniqueNameProvider ) );
 		ChildActorComponent->SetChildActorClass( TSubclassOf< AActor > ( LightShapeBlueprint->GeneratedClass ) );
 		ChildActorComponent->CreateChildActor();
 
@@ -474,9 +474,9 @@ USceneComponent* FDatasmithLightImporter::ImportAreaLightComponent( const TShare
 	return MainComponent;
 }
 
-ULightmassPortalComponent* FDatasmithLightImporter::ImportLightmassPortalComponent( const TSharedRef< IDatasmithLightmassPortalElement >& LightElement, FDatasmithImportContext& ImportContext, UObject* Outer )
+ULightmassPortalComponent* FDatasmithLightImporter::ImportLightmassPortalComponent( const TSharedRef< IDatasmithLightmassPortalElement >& LightElement, FDatasmithImportContext& ImportContext, UObject* Outer, FDatasmithActorUniqueLabelProvider& UniqueNameProvider )
 {
-	USceneComponent* SceneComponent = FDatasmithActorImporter::ImportSceneComponent( ULightmassPortalComponent::StaticClass(), LightElement, ImportContext, Outer );
+	USceneComponent* SceneComponent = FDatasmithActorImporter::ImportSceneComponent( ULightmassPortalComponent::StaticClass(), LightElement, ImportContext, Outer, UniqueNameProvider );
 
 	SceneComponent->RegisterComponent();
 
diff --git a/Engine/Plugins/Enterprise/DatasmithImporter/Source/DatasmithImporter/Public/DatasmithActorImporter.h b/Engine/Plugins/Enterprise/DatasmithImporter/Source/DatasmithImporter/Public/DatasmithActorImporter.h
--- a/Engine/Plugins/Enterprise/DatasmithImporter/Source/DatasmithImporter/Public/DatasmithActorImporter.h
+++ b/Engine/Plugins/Enterprise/DatasmithImporter/Source/DatasmithImporter/Public/DatasmithActorImporter.h
@@ -26,6 +26,7 @@ class UHierarchicalInstancedStaticMeshComponent;
 class USceneComponent;
 class UStaticMeshComponent;
 struct FCachedActorLabels;
+class FDatasmithActorUniqueLabelProvider;
 
 class FDatasmithActorImporter
 {
@@ -55,19 +56,19 @@ public:
 	 *
 	 * @return	Returns the importer scene component. The component is not registered to allow edits. The caller needs to register it when the setup is complete.
 	 */
-	static USceneComponent* ImportSceneComponent( UClass* ComponentClass, const TSharedRef< IDatasmithActorElement >& ActorElement, FDatasmithImportContext& ImportContext, UObject* Outer );
+	static USceneComponent* ImportSceneComponent( UClass* ComponentClass, const TSharedRef< IDatasmithActorElement >& ActorElement, FDatasmithImportContext& ImportContext, UObject* Outer, FDatasmithActorUniqueLabelProvider& UniqueNameProvider );
 
 	/**
 	 * Spawns an actor with a SceneComponent
 	 */
 	static AActor* ImportBaseActor( FDatasmithImportContext& ImportContext, const TSharedRef< IDatasmithActorElement >& ActorElement );
-	static USceneComponent* ImportBaseActorAsComponent( FDatasmithImportContext& ImportContext, const TSharedRef< IDatasmithActorElement >& ActorElement, UObject* Outer );
+	static USceneComponent* ImportBaseActorAsComponent( FDatasmithImportContext& ImportContext, const TSharedRef< IDatasmithActorElement >& ActorElement, UObject* Outer, FDatasmithActorUniqueLabelProvider& UniqueNameProvider );
 
 	/**
 	 * Spawns a static mesh actor
 	 */
 	static AStaticMeshActor* ImportStaticMeshActor( FDatasmithImportContext& ImportContext, const TSharedRef< IDatasmithMeshActorElement >& MeshActorElement );
-	static UStaticMeshComponent* ImportStaticMeshComponent( FDatasmithImportContext& ImportContext, const TSharedRef< IDatasmithMeshActorElement >& MeshActorElement, UObject* Outer );
+	static UStaticMeshComponent* ImportStaticMeshComponent( FDatasmithImportContext& ImportContext, const TSharedRef< IDatasmithMeshActorElement >& MeshActorElement, UObject* Outer, FDatasmithActorUniqueLabelProvider& UniqueNameProvider );
 
 	/**
 	 * Spawns a cine camera actor
@@ -87,7 +88,7 @@ public:
 	/**
 	 * Spawns an actor with a class defined by the CustomActorElement class name
 	 */
-	static AActor* ImportCustomActor( FDatasmithImportContext& ImportContext, const TSharedRef< IDatasmithCustomActorElement >& CustomActorElement );
+	static AActor* ImportCustomActor( FDatasmithImportContext& ImportContext, const TSharedRef< IDatasmithCustomActorElement >& CustomActorElement, FDatasmithActorUniqueLabelProvider& UniqueNameProvider );
 
 	/**
 	 * Spawns a landscape actor
@@ -97,10 +98,10 @@ public:
 	/**
 	 * Spawns an actor with a hierarchical instanced static mesh component
 	 */
-	static AActor* ImportHierarchicalInstancedStaticMeshAsActor( FDatasmithImportContext& ImportContext, const TSharedRef< IDatasmithHierarchicalInstancedStaticMeshActorElement >& HierarchicalInstancedStatictMeshActorElement );
+	static AActor* ImportHierarchicalInstancedStaticMeshAsActor( FDatasmithImportContext& ImportContext, const TSharedRef< IDatasmithHierarchicalInstancedStaticMeshActorElement >& HierarchicalInstancedStatictMeshActorElement, FDatasmithActorUniqueLabelProvider& UniqueNameProvider );
 	
 	static UHierarchicalInstancedStaticMeshComponent* ImportHierarchicalInstancedStaticMeshComponent( FDatasmithImportContext& ImportContext,
-		const TSharedRef< IDatasmithHierarchicalInstancedStaticMeshActorElement >& HierarchicalInstancedStaticMeshActorElement, UObject* Outer );
+		const TSharedRef< IDatasmithHierarchicalInstancedStaticMeshActorElement >& HierarchicalInstancedStaticMeshActorElement, UObject* Outer, FDatasmithActorUniqueLabelProvider& UniqueNameProvider );
 
 	/**
 	 * Handle Actor common properties (Layers, Tags)
diff --git a/Engine/Plugins/Enterprise/DatasmithImporter/Source/DatasmithImporter/Public/DatasmithCameraImporter.h b/Engine/Plugins/Enterprise/DatasmithImporter/Source/DatasmithImporter/Public/DatasmithCameraImporter.h
--- a/Engine/Plugins/Enterprise/DatasmithImporter/Source/DatasmithImporter/Public/DatasmithCameraImporter.h
+++ b/Engine/Plugins/Enterprise/DatasmithImporter/Source/DatasmithImporter/Public/DatasmithCameraImporter.h
@@ -8,6 +8,7 @@ class AActor;
 struct FDatasmithImportContext;
 class IDatasmithCameraActorElement;
 class UCineCameraComponent;
+class FDatasmithActorUniqueLabelProvider;
 
 class FDatasmithCameraImporter
 {
@@ -15,7 +16,7 @@ public:
 	static AActor* ImportCameraActor( const TSharedRef< IDatasmithCameraActorElement >& CameraElement, FDatasmithImportContext& ImportContext );
 	static void PostImportCameraActor( const TSharedRef< IDatasmithCameraActorElement >& CameraElement, FDatasmithImportContext& ImportContext );
 
-	static UCineCameraComponent* ImportCineCameraComponent( const TSharedRef< IDatasmithCameraActorElement >& CameraElement, FDatasmithImportContext& ImportContext, UObject* Outer );
+	static UCineCameraComponent* ImportCineCameraComponent( const TSharedRef< IDatasmithCameraActorElement >& CameraElement, FDatasmithImportContext& ImportContext, UObject* Outer, FDatasmithActorUniqueLabelProvider& UniqueNameProvider );
 
 protected:
 	static void SetupCineCameraComponent( UCineCameraComponent* CineCameraComponent, const TSharedRef< IDatasmithCameraActorElement >& CameraElement );
diff --git a/Engine/Plugins/Enterprise/DatasmithImporter/Source/DatasmithImporter/Public/DatasmithImporter.h b/Engine/Plugins/Enterprise/DatasmithImporter/Source/DatasmithImporter/Public/DatasmithImporter.h
--- a/Engine/Plugins/Enterprise/DatasmithImporter/Source/DatasmithImporter/Public/DatasmithImporter.h
+++ b/Engine/Plugins/Enterprise/DatasmithImporter/Source/DatasmithImporter/Public/DatasmithImporter.h
@@ -140,7 +140,7 @@ struct DATASMITHIMPORTER_API FDatasmithImporter
 	 * @param ActorElement		The Datasmith element containing the data for the component
 	 * @param RootActor			The actor in which to add the component
 	 */
-	static void ImportActorAsComponent( FDatasmithImportContext& ImportContext, const TSharedRef< IDatasmithActorElement >& ActorElement, AActor* InRootActor );
+	static void ImportActorAsComponent( FDatasmithImportContext& ImportContext, const TSharedRef< IDatasmithActorElement >& ActorElement, AActor* InRootActor, FDatasmithActorUniqueLabelProvider& UniqueNameProvider );
 
 	/**
 	 * Finalizes all the actors
diff --git a/Engine/Plugins/Enterprise/DatasmithImporter/Source/DatasmithImporter/Public/DatasmithLightImporter.h b/Engine/Plugins/Enterprise/DatasmithImporter/Source/DatasmithImporter/Public/DatasmithLightImporter.h
--- a/Engine/Plugins/Enterprise/DatasmithImporter/Source/DatasmithImporter/Public/DatasmithLightImporter.h
+++ b/Engine/Plugins/Enterprise/DatasmithImporter/Source/DatasmithImporter/Public/DatasmithLightImporter.h
@@ -32,12 +32,13 @@ class USpotLightComponent;
 class UTexture;
 class UTextureLightProfile;
 class UWorld;
+class FDatasmithActorUniqueLabelProvider;
 
 class DATASMITHIMPORTER_API FDatasmithLightImporter
 {
 public:
 	static AActor* ImportLightActor( const TSharedRef< IDatasmithLightActorElement >& LightElement, FDatasmithImportContext& ImportContext );
-	static USceneComponent* ImportLightComponent( const TSharedRef< IDatasmithLightActorElement >& LightElement, FDatasmithImportContext& ImportContext, UObject* Outer );
+	static USceneComponent* ImportLightComponent( const TSharedRef< IDatasmithLightActorElement >& LightElement, FDatasmithImportContext& ImportContext, UObject* Outer, FDatasmithActorUniqueLabelProvider& UniqueNameProvider );
 
 	static AActor* CreateHDRISkyLight(const TSharedPtr< IDatasmithShaderElement >& ShaderElement, FDatasmithImportContext& ImportContext);
 
@@ -53,11 +54,11 @@ private:
 	static void SetupSpotLightComponent( USpotLightComponent* SpotLightComponent, const TSharedRef< IDatasmithSpotLightElement >& SpotLightElement, const TCHAR* LightsFolderPath, const TCHAR* MaterialsFolderPath );
 
 	static AActor* ImportAreaLightActor( const TSharedRef< IDatasmithAreaLightElement >& AreaLightElement, FDatasmithImportContext& ImportContext );
-	static USceneComponent* ImportAreaLightComponent( const TSharedRef< IDatasmithAreaLightElement >& AreaLightElement, FDatasmithImportContext& ImportContext, UObject* Outer );
+	static USceneComponent* ImportAreaLightComponent( const TSharedRef< IDatasmithAreaLightElement >& AreaLightElement, FDatasmithImportContext& ImportContext, UObject* Outer, FDatasmithActorUniqueLabelProvider& UniqueNameProvider );
 
 	static AActor* CreateAreaLightActor( const TSharedRef< IDatasmithAreaLightElement >& AreaLightElement, FDatasmithImportContext& ImportContext );
 	static void SetupAreaLightActor( const TSharedRef< IDatasmithAreaLightElement >& AreaLightElement, FDatasmithImportContext& ImportContext, ADatasmithAreaLightActor* LightShapeActor );
-	static ULightmassPortalComponent* ImportLightmassPortalComponent( const TSharedRef< IDatasmithLightmassPortalElement >& LightElement, FDatasmithImportContext& ImportContext, UObject* Outer );
+	static ULightmassPortalComponent* ImportLightmassPortalComponent( const TSharedRef< IDatasmithLightmassPortalElement >& LightElement, FDatasmithImportContext& ImportContext, UObject* Outer, FDatasmithActorUniqueLabelProvider& UniqueNameProvider );
 	static AActor* CreateSkyLight(const TSharedPtr< IDatasmithShaderElement >& ShaderElement, FDatasmithImportContext& ImportContext, bool bUseHDRMat);
 
 };
