Integrating CL# 11420788 from Dev-Enterprise
Fixing Datasmith Static Mesh reimport material override

Danny.Couture JeanMichel.Dignard

diff --git a/Engine/Plugins/Enterprise/DatasmithImporter/Source/DatasmithImporter/Private/DatasmithStaticMeshImporter.cpp b/Engine/Plugins/Enterprise/DatasmithImporter/Source/DatasmithImporter/Private/DatasmithStaticMeshImporter.cpp
--- a/Engine/Plugins/Enterprise/DatasmithImporter/Source/DatasmithImporter/Private/DatasmithStaticMeshImporter.cpp
+++ b/Engine/Plugins/Enterprise/DatasmithImporter/Source/DatasmithImporter/Private/DatasmithStaticMeshImporter.cpp
@@ -124,7 +124,9 @@ UStaticMesh* FDatasmithStaticMeshImporter::ImportStaticMesh(const TSharedRef< ID
 	{
 		if ( ExistingMesh->GetOuter() != Outer )
 		{
-			ResultStaticMesh = FDatasmithImporterUtils::DuplicateObject< UStaticMesh >( ExistingMesh, Outer, *StaticMeshName );
+			// We don't need to copy over the mesh BulkData as it is going to be recreated anyway, this also prevent ExistingMesh from being invalidated.
+			const bool bIgnoreBulkData = true;
+			ResultStaticMesh = FDatasmithImporterUtils::DuplicateStaticMesh( ExistingMesh, Outer, *StaticMeshName, bIgnoreBulkData);
 			IDatasmithImporterModule::Get().ResetOverrides( ResultStaticMesh ); // Don't copy the existing overrides
 		}
 		else
@@ -133,17 +135,6 @@ UStaticMesh* FDatasmithStaticMeshImporter::ImportStaticMesh(const TSharedRef< ID
 		}
 
 		ResultStaticMesh->SetFlags( ObjectFlags );
-
-		// Get rid of data that is going to be recreated anyway
-		for (FStaticMeshSourceModel& SourceModel : ResultStaticMesh->GetSourceModels())
-		{
-			SourceModel.MeshDescription.Reset();
-			SourceModel.MeshDescriptionBulkData.Reset();
-			if (SourceModel.RawMeshBulkData)
-			{
-				SourceModel.RawMeshBulkData->Empty();
-			}
-		}
 	}
 	else
 	{
diff --git a/Engine/Plugins/Enterprise/DatasmithImporter/Source/DatasmithImporter/Private/Utility/DatasmithImporterUtils.cpp b/Engine/Plugins/Enterprise/DatasmithImporter/Source/DatasmithImporter/Private/Utility/DatasmithImporterUtils.cpp
--- a/Engine/Plugins/Enterprise/DatasmithImporter/Source/DatasmithImporter/Private/Utility/DatasmithImporterUtils.cpp
+++ b/Engine/Plugins/Enterprise/DatasmithImporter/Source/DatasmithImporter/Private/Utility/DatasmithImporterUtils.cpp
@@ -1088,52 +1088,79 @@ const TSharedPtr<IDatasmithBaseMaterialElement>& FDatasmithImporterUtils::FDatas
 
 UObject* FDatasmithImporterUtils::StaticDuplicateObject(UObject* SourceObject, UObject* Outer, const FName Name)
 {
-	UStaticMesh* SourceMesh = Cast< UStaticMesh >( SourceObject );
+	if (UStaticMesh* SourceMesh = Cast< UStaticMesh >(SourceObject))
+	{
+		const bool bIgnoreBulkData = false;
+		return DuplicateStaticMesh(SourceMesh, Outer, Name, bIgnoreBulkData);
+	}
 
+	// Duplicate is used only to move our object from its temporary package into its final package replacing any asset
+	// already at that location. This function also takes care of fixing internal dependencies among the object's children.
+	// Since Duplicate has some rather heavy consequence, like calling PostLoad and doing all kind of stuff on an object
+	// that is not even fully initialized yet, we might want to find an alternative way of moving our objects in future
+	// releases but keep it for the current release cycle.
+	return ::DuplicateObject< UObject >( SourceObject, Outer, Name );
+}
+
+UStaticMesh* FDatasmithImporterUtils::DuplicateStaticMesh(UStaticMesh* SourceStaticMesh, UObject* Outer, const FName Name, bool bIgnoreBulkData)
+{
 	TArray<FStaticMeshSourceModel> SourceModels;
-	if (SourceMesh)
-	{
-		// Since static mesh can be quite heavy, remove source models for cloning to
-		// reduce useless work. Will be reinserted on the new duplicated asset.
-		SourceModels = MoveTemp(SourceMesh->GetSourceModels());
 
-		// Temporary flag to skip Postload during DuplicateObject
-		SourceMesh->SetFlags(RF_ArchetypeObject);
-	}
+	// Since static mesh can be quite heavy, remove source models for cloning to reduce useless work.
+	// Will be reinserted on the new duplicated asset or restored on the SourceStaticMesh if bIgnoreBulkData is true.
+	SourceModels = MoveTemp(SourceStaticMesh->GetSourceModels());
+
+	// Temporary flag to skip Postload during DuplicateObject
+	SourceStaticMesh->SetFlags(RF_ArchetypeObject);
 
 	// Duplicate is used only to move our object from its temporary package into its final package replacing any asset
 	// already at that location. This function also takes care of fixing internal dependencies among the object's children.
 	// Since Duplicate has some rather heavy consequence, like calling PostLoad and doing all kind of stuff on an object
 	// that is not even fully initialized yet, we might want to find an alternative way of moving our objects in future
 	// releases but keep it for the current release cycle.
-	UObject* Duplicate = ::DuplicateObject< UObject >( SourceObject, Outer, Name );
+	UStaticMesh* DuplicateMesh = ::DuplicateObject< UStaticMesh >(SourceStaticMesh, Outer, Name);
 
-	if ( SourceMesh )
-	{
-		UStaticMesh* DuplicateMesh = Cast< UStaticMesh >( Duplicate );
+	// Get rid of our temporary flag
+	SourceStaticMesh->ClearFlags(RF_ArchetypeObject);
+	DuplicateMesh->ClearFlags(RF_ArchetypeObject);
 
-		// Get rid of our temporary flag
-		SourceMesh->ClearFlags(RF_ArchetypeObject);
-		DuplicateMesh->ClearFlags(RF_ArchetypeObject);
+	if (bIgnoreBulkData)
+	{
+		//We are not moving the source model bulk data, so we can simply copy the settings used in the DDC key.
+		//That way we also avoid marking the source mesh as pending kill.
+		for (FStaticMeshSourceModel& SourceModel : SourceModels)
+		{
+			FStaticMeshSourceModel& DuplicateSourceModel = DuplicateMesh->AddSourceModel();
+			DuplicateSourceModel.StaticMeshOwner = DuplicateMesh;
+
+			// Apply the SourceMesh settings to the duplicated SourceModels
+			DuplicateSourceModel.BuildSettings = SourceModel.BuildSettings;
+			DuplicateSourceModel.ReductionSettings = SourceModel.ReductionSettings;
+			DuplicateSourceModel.ScreenSize = SourceModel.ScreenSize;
+			DuplicateSourceModel.bImportWithBaseMesh = SourceModel.bImportWithBaseMesh;
+			DuplicateSourceModel.SourceImportFilename = SourceModel.SourceImportFilename;
+		}
 
+		// Move back the source models to the original mesh
+		SourceStaticMesh->GetSourceModels() = MoveTemp(SourceModels);
+	}
+	else
+	{
 		// The source mesh is stripped from it's source model, it is not buildable anymore.
 		// -> MarkPendingKill to avoid use-after-move crash in the StaticMesh::Build()
-		SourceMesh->MarkPendingKill();
+		SourceStaticMesh->MarkPendingKill();
 
 		for (FStaticMeshSourceModel& SourceModel : SourceModels)
 		{
 			// Fixup the new SourceModels owner
 			SourceModel.StaticMeshOwner = DuplicateMesh;
-
-			// Restore settings to the SourceMesh since they are used for diffing templates
-			SourceMesh->AddSourceModel().BuildSettings = SourceModel.BuildSettings;
 		}
 
 		// Apply source models to the duplicated mesh
 		DuplicateMesh->GetSourceModels() = MoveTemp(SourceModels);
 	}
 
-	return Duplicate;
+	return DuplicateMesh;
 }
 
 FScopedLogger::FScopedLogger(FName LogTitle, const FText& LogLabel)
diff --git a/Engine/Plugins/Enterprise/DatasmithImporter/Source/DatasmithImporter/Public/Utility/DatasmithImporterUtils.h b/Engine/Plugins/Enterprise/DatasmithImporter/Source/DatasmithImporter/Public/Utility/DatasmithImporterUtils.h
--- a/Engine/Plugins/Enterprise/DatasmithImporter/Source/DatasmithImporter/Public/Utility/DatasmithImporterUtils.h
+++ b/Engine/Plugins/Enterprise/DatasmithImporter/Source/DatasmithImporter/Public/Utility/DatasmithImporterUtils.h
@@ -170,6 +170,19 @@ public:
 	 */
 	static UObject* StaticDuplicateObject(UObject* SourceObject, UObject* Outer, const FName Name = NAME_None);
 
+	/**
+	 * Specialization of the duplication of a StaticMesh object specifically optimized for datasmith use case.
+	 * This operation invalidate the duplicated SourceStaticMesh and mark it as PendingKill, unless bIgnoreBulkData is True.
+	 *
+	 * @param SourceStaticMesh	the UStaticMesh being copied
+	 * @param Outer				the outer to use for the object
+	 * @param Name				the optional name of the object
+	 * @param bIgnoreBulkData	if True, the SourceStaticMesh's SourceModels BulkDatas won't be copied and SourceStaticMesh will stay valid after the operation.
+	 *
+	 * @return the copied StaticMesh or null if it failed for some reason
+	 */
+	static UStaticMesh* DuplicateStaticMesh(UStaticMesh* SourceStaticMesh, UObject* Outer, const FName Name = NAME_None, bool bIgnoreBulkData = false);
+
 	template< class T >
 	static T* DuplicateObject(T* SourceObject, UObject* Outer, const FName Name = NAME_None)
 	{
