Integrating CL # from //UE4/Dev-Enterprise

Fixed name uniqueness issue when loading files containing assets with the same name

diff --git a/Engine/Plugins/Enterprise/DatasmithImporter/Source/DatasmithImporter/Private/DatasmithFileProducer.cpp b/Engine/Plugins/Enterprise/DatasmithImporter/Source/DatasmithImporter/Private/DatasmithFileProducer.cpp
--- a/Engine/Plugins/Enterprise/DatasmithImporter/Source/DatasmithImporter/Private/DatasmithFileProducer.cpp
+++ b/Engine/Plugins/Enterprise/DatasmithImporter/Source/DatasmithImporter/Private/DatasmithFileProducer.cpp
@@ -421,7 +421,9 @@ void UDatasmithFileProducer::PreventNameCollision()
 		// Map containing mapping between previous package to new one
 		TMap<FSoftObjectPath, FSoftObjectPath> AssetRedirectorMap;
 
-		auto MoveAsset = [&AssetRedirectorMap, &PackagesToCheck](UObject* Object, UPackage* NewPackage, bool bCheckPackage)
+		UDatasmithContentBlueprintLibrary* DatasmithContentLibrary = Cast< UDatasmithContentBlueprintLibrary >( UDatasmithContentBlueprintLibrary::StaticClass()->GetDefaultObject() );
+
+		auto MoveAsset = [&DatasmithContentLibrary, &AssetRedirectorMap, &PackagesToCheck](UObject* Object, UPackage* NewPackage, bool bCheckPackage)
 		{
 			if(Object->GetOutermost()->GetName() != NewPackage->GetName())
 			{
@@ -429,6 +431,14 @@ void UDatasmithFileProducer::PreventNameCollision()
 
 				Object->Rename( nullptr, NewPackage, REN_DontCreateRedirectors | REN_NonTransactional );
 
+				// Replace unique name id with hash of package path to avoid asset's name collision
+				FString DatasmithUniqueId = DatasmithContentLibrary->GetDatasmithUserDataValueForKey( Object, UDatasmithAssetUserData::UniqueIdMetaDataKey );
+				const FString ObjectPath = Object->GetPathName();
+				FString NewUniqueId = FMD5::HashBytes( reinterpret_cast<const uint8*>(*ObjectPath), ObjectPath.Len() * sizeof(TCHAR) );
+
+				UDatasmithAssetUserData::SetDatasmithUserDataValueForKey(Object, *(FString(TEXT("Old")) + UDatasmithAssetUserData::UniqueIdMetaDataKey), DatasmithUniqueId );
+				UDatasmithAssetUserData::SetDatasmithUserDataValueForKey(Object, UDatasmithAssetUserData::UniqueIdMetaDataKey, NewUniqueId );
+
 				AssetRedirectorMap.Emplace( PreviousObjectPath, Object );
 				if(bCheckPackage)
 				{
