Integrating CL# 11509121 from Dev-Enterprise
Fixing datasmith light template not applied on reimport

JeanMichel.Dignard

diff --git a/Engine/Plugins/Enterprise/DatasmithImporter/Source/DatasmithImporter/Private/DatasmithImporter.cpp b/Engine/Plugins/Enterprise/DatasmithImporter/Source/DatasmithImporter/Private/DatasmithImporter.cpp
--- a/Engine/Plugins/Enterprise/DatasmithImporter/Source/DatasmithImporter/Private/DatasmithImporter.cpp
+++ b/Engine/Plugins/Enterprise/DatasmithImporter/Source/DatasmithImporter/Private/DatasmithImporter.cpp
@@ -889,6 +889,61 @@ namespace DatasmithImporterImpl
 			AssetTools.ConvertVirtualTextures(TexturesToConvertList, true, &MaterialsToRefreshAfterVirtualTextureConversion);
 		}
 	}
+
+	/**
+	 * Helper struct used to manage and validate the actor's state when modifying it during the finalization.
+	 */
+	struct FScopedFinalizeActorChanges
+	{
+		FDatasmithImportContext& ImportContext;
+		AActor* FinalizedActor;
+
+		FScopedFinalizeActorChanges(AActor* InFinalizedActor, FDatasmithImportContext& InImportContext)
+			: ImportContext(InImportContext)
+			, FinalizedActor(InFinalizedActor)
+		{
+			// In order to allow modification on components owned by ExistingActor, unregister all of them
+			FinalizedActor->UnregisterAllComponents( /* bForReregister = */true);
+		}
+
+		~FScopedFinalizeActorChanges()
+		{
+			for (UActorComponent* Component : FinalizedActor->GetComponents())
+			{
+				if (Component->IsRegistered())
+				{
+					ensureMsgf(false, TEXT("All components should still be unregistered at this point. Otherwise some datasmith templates might not have been applied properly."));
+					break;
+				}
+			}
+			
+			const FQuat PreviousRotation(FinalizedActor->GetRootComponent()->GetRelativeTransform().GetRotation());
+			FinalizedActor->PostEditChange();
+			FinalizedActor->RegisterAllComponents();
+
+			const bool bHasPostEditChangeModifiedRotation = !PreviousRotation.Equals(FinalizedActor->GetRootComponent()->GetRelativeTransform().GetRotation());
+			if (bHasPostEditChangeModifiedRotation)
+			{
+				//SingularityThreshold value is comming from the FQuat::Rotator() function, but is more permissive because the rotation is already diverging before the singularity threshold is reached.
+				const float SingularityThreshold = 0.4999f;
+				const float SingularityTest = PreviousRotation.Z * PreviousRotation.X - PreviousRotation.W * PreviousRotation.Y;
+				const AActor* RootSceneActor = ImportContext.ActorsContext.ImportSceneActor;
+
+				if (FinalizedActor != RootSceneActor
+					&& FMath::Abs(SingularityTest) > SingularityThreshold)
+				{
+					//This is a warning to explain the edge-case of UE-75467 while it's being fixed.
+					FFormatNamedArguments FormatArgs;
+					FormatArgs.Add(TEXT("ActorName"), FText::FromName(FinalizedActor->GetFName()));
+					ImportContext.LogWarning(FText::GetEmpty())
+						->AddToken(FUObjectToken::Create(FinalizedActor))
+						->AddToken(FTextToken::Create(FText::Format(LOCTEXT("UnsupportedRotationValueError", "The actor '{ActorName}' has a rotation value pointing to either (0, 90, 0) or (0, -90, 0)."
+							"This is an edge case that is not well supported in Unreal and can cause incorrect results."
+							"In those cases, it is recommended to bake the actor's transform into the mesh at export."), FormatArgs)));
+				}
+			}
+		}
+	};
 }
 
 void FDatasmithImporter::ImportStaticMeshes( FDatasmithImportContext& ImportContext )
@@ -1846,96 +1901,70 @@ AActor* FDatasmithImporter::FinalizeActor( FDatasmithImportContext& ImportContex
 		ExistingActor = nullptr;
 	}
 
-	// Backup hierarchy
 	TArray< AActor* > Children;
-	if ( ExistingActor )
-	{
-		ExistingActor->GetAttachedActors( Children );
-
-		// In order to allow modification on components owned by ExistingActor, unregister all of them
-		ExistingActor->UnregisterAllComponents( /* bForReregister = */true );
-	}
-
 	AActor* DestinationActor = ExistingActor;
 	if ( !DestinationActor )
 	{
 		DestinationActor = ImportContext.ActorsContext.FinalWorld->SpawnActor( SourceActor.GetClass() );
 	}
+	else
+	{
+		// Backup hierarchy
+		ExistingActor->GetAttachedActors( Children );
+	}
 
 	// Update label to match the source actor's
 	DestinationActor->SetActorLabel( ImportContext.ActorsContext.UniqueNameProvider.GenerateUniqueName( SourceActor.GetActorLabel() ) );
 
 	check( DestinationActor );
-	ReferencesToRemap.Add( &SourceActor ) = DestinationActor;
-
-	TArray< FMigratedTemplatePairType > MigratedTemplates = MigrateTemplates(
-		&SourceActor,
-		ExistingActor,
-		&ReferencesToRemap,
-		true
-	);
-
-	// Copy actor data
 	{
-		TArray< uint8 > Bytes;
-		FActorWriter ObjectWriter( &SourceActor, Bytes );
-		FObjectReader ObjectReader( DestinationActor, Bytes );
-	}
+		// Setup the actor to allow modifications.
+		FScopedFinalizeActorChanges ScopedFinalizedActorChanges(DestinationActor, ImportContext);
+		
+		ReferencesToRemap.Add( &SourceActor ) = DestinationActor;
 
-	FixReferencesForObject( DestinationActor, ReferencesToRemap );
+		TArray< FMigratedTemplatePairType > MigratedTemplates = MigrateTemplates(
+			&SourceActor,
+			ExistingActor,
+			&ReferencesToRemap,
+			true
+		);
 
-	DatasmithImporterImpl::FinalizeComponents( ImportContext, SourceActor, *DestinationActor, ReferencesToRemap );
-
-	// The templates for the actor need to be applied after the components were created.
-	ApplyMigratedTemplates( MigratedTemplates, DestinationActor );
+		// Copy actor data
+		{
+			TArray< uint8 > Bytes;
+			FActorWriter ObjectWriter( &SourceActor, Bytes );
+			FObjectReader ObjectReader( DestinationActor, Bytes );
+		}
 
-	// Restore hierarchy
-	for ( AActor* Child : Children )
-	{
-		Child->AttachToActor( DestinationActor, FAttachmentTransformRules::KeepWorldTransform );
-	}
+		FixReferencesForObject( DestinationActor, ReferencesToRemap );
 
-	// Hotfix for UE-69555
-	TInlineComponentArray<UHierarchicalInstancedStaticMeshComponent*> HierarchicalInstancedStaticMeshComponents;
-	DestinationActor->GetComponents(HierarchicalInstancedStaticMeshComponents);
-	for (UHierarchicalInstancedStaticMeshComponent* HierarchicalInstancedStaticMeshComponent : HierarchicalInstancedStaticMeshComponents )
-	{
-		HierarchicalInstancedStaticMeshComponent->BuildTreeIfOutdated( true, true );
-	}
+		DatasmithImporterImpl::FinalizeComponents( ImportContext, SourceActor, *DestinationActor, ReferencesToRemap );
 
-	if ( ALandscape* Landscape = Cast< ALandscape >( DestinationActor ) )
-	{
-		FPropertyChangedEvent MaterialPropertyChangedEvent( FindFieldChecked< FProperty >( Landscape->GetClass(), FName("LandscapeMaterial") ) );
-		Landscape->PostEditChangeProperty( MaterialPropertyChangedEvent );
-	}
+		// The templates for the actor need to be applied after the components were created.
+		ApplyMigratedTemplates( MigratedTemplates, DestinationActor );
 
-	FQuat PreviousRotation = DestinationActor->GetRootComponent()->GetRelativeTransform().GetRotation();
-	DestinationActor->PostEditChange();
+		// Restore hierarchy
+		for ( AActor* Child : Children )
+		{
+			Child->AttachToActor( DestinationActor, FAttachmentTransformRules::KeepWorldTransform );
+		}
 
-	const bool bHasPostEditChangeModifiedRotation = !PreviousRotation.Equals(DestinationActor->GetRootComponent()->GetRelativeTransform().GetRotation());
-	if (bHasPostEditChangeModifiedRotation)
-	{
-		const float SingularityTest = PreviousRotation.Z * PreviousRotation.X - PreviousRotation.W * PreviousRotation.Y;
-		//SingularityThreshold value is comming from the FQuat::Rotator() function, but is more permissive because the rotation is already diverging before the singularity threshold is reached.
-		const float SingularityThreshold = 0.4999f; 
+		// Hotfix for UE-69555
+		TInlineComponentArray<UHierarchicalInstancedStaticMeshComponent*> HierarchicalInstancedStaticMeshComponents;
+		DestinationActor->GetComponents(HierarchicalInstancedStaticMeshComponents);
+		for (UHierarchicalInstancedStaticMeshComponent* HierarchicalInstancedStaticMeshComponent : HierarchicalInstancedStaticMeshComponents )
+		{
+			HierarchicalInstancedStaticMeshComponent->BuildTreeIfOutdated( true, true );
+		}
 
-		AActor* RootSceneActor = ImportContext.ActorsContext.ImportSceneActor;
-		if (DestinationActor != RootSceneActor
-			&& FMath::Abs(SingularityTest) > SingularityThreshold)
+		if ( ALandscape* Landscape = Cast< ALandscape >( DestinationActor ) )
 		{
-			//This is a warning to explain the edge-case of UE-75467 while it's being fixed.
-			FFormatNamedArguments FormatArgs;
-			FormatArgs.Add(TEXT("ActorName"), FText::FromName(DestinationActor->GetFName()));
-			ImportContext.LogWarning(FText::GetEmpty())
-				->AddToken(FUObjectToken::Create(DestinationActor))
-				->AddToken(FTextToken::Create(FText::Format(LOCTEXT("UnsupportedRotationValueError", "The actor '{ActorName}' has a rotation value pointing to either (0, 90, 0) or (0, -90, 0)."
-					"This is an edge case that is not well supported in Unreal and can cause incorrect results."
-					"In those cases, it is recommended to bake the actor's transform into the mesh at export."), FormatArgs)));
+			FPropertyChangedEvent MaterialPropertyChangedEvent( FindFieldChecked< FProperty >( Landscape->GetClass(), FName("LandscapeMaterial") ) );
+			Landscape->PostEditChangeProperty( MaterialPropertyChangedEvent );
 		}
 	}
 
-	DestinationActor->RegisterAllComponents();
-
 	return DestinationActor;
 }
 
