Integrating CL# 11541462 from Dev-Enterprise
Fixing registered component ensure

diff --git a/Engine/Plugins/Enterprise/DatasmithImporter/Source/DatasmithImporter/Private/DatasmithImporter.cpp b/Engine/Plugins/Enterprise/DatasmithImporter/Source/DatasmithImporter/Private/DatasmithImporter.cpp
--- a/Engine/Plugins/Enterprise/DatasmithImporter/Source/DatasmithImporter/Private/DatasmithImporter.cpp
+++ b/Engine/Plugins/Enterprise/DatasmithImporter/Source/DatasmithImporter/Private/DatasmithImporter.cpp
@@ -896,6 +896,7 @@ namespace DatasmithImporterImpl
 	struct FScopedFinalizeActorChanges
 	{
 		FDatasmithImportContext& ImportContext;
+		TSet<UActorComponent*> ComponentsToValidate;
 		AActor* FinalizedActor;
 
 		FScopedFinalizeActorChanges(AActor* InFinalizedActor, FDatasmithImportContext& InImportContext)
@@ -904,13 +905,19 @@ namespace DatasmithImporterImpl
 		{
 			// In order to allow modification on components owned by ExistingActor, unregister all of them
 			FinalizedActor->UnregisterAllComponents( /* bForReregister = */true);
+
+			//Some new components might be created when finalizing the actor, only validate those that we unregistered.
+			for (UActorComponent* Component : FinalizedActor->GetComponents())
+			{
+				ComponentsToValidate.Add(Component);
+			}
 		}
 
 		~FScopedFinalizeActorChanges()
 		{
 			for (UActorComponent* Component : FinalizedActor->GetComponents())
 			{
-				if (Component->IsRegistered())
+				if (Component->IsRegistered() && ComponentsToValidate.Contains(Component))
 				{
 					ensureMsgf(false, TEXT("All components should still be unregistered at this point. Otherwise some datasmith templates might not have been applied properly."));
 					break;
