Added dedicated UAssetUserData for DataprepContentConsumer to handle meta-data for consumers
Reworked UDatasmithConsumer to handle the creation/edition of levels at execution
Added a boolean to the consumer's context to prevent the opening of dialogs when executing a consumer in a script
Removed the DataprepConsumerWidget
Added operations to output on different levels and/or folders
Added a selection transform to gather immediate or all descendants of an actor
Completed the fix to track assets and actors with the same Datasmith uniqueID

diff --git a/Engine/Plugins/Enterprise/DataprepEditor/Source/DataprepCore/Private/DataprepAsset.cpp b/Engine/Plugins/Enterprise/DataprepEditor/Source/DataprepCore/Private/DataprepAsset.cpp
--- a/Engine/Plugins/Enterprise/DataprepEditor/Source/DataprepCore/Private/DataprepAsset.cpp
+++ b/Engine/Plugins/Enterprise/DataprepEditor/Source/DataprepCore/Private/DataprepAsset.cpp
@@ -45,8 +45,10 @@ void UDataprepAsset::PostLoad()
 {
 	UDataprepAssetInterface::PostLoad();
 
+#ifndef NO_BLUEPRINT
 	check(DataprepRecipeBP);
 	DataprepRecipeBP->OnChanged().AddUObject( this, &UDataprepAsset::OnDataprepBlueprintChanged );
+#endif
 
 	// Move content of deprecated properties to the corresponding new ones.
 	if(HasAnyFlags(RF_WasLoaded))
@@ -156,6 +158,15 @@ void UDataprepAsset::PostEditUndo()
 	CachedActionCount = ActionAssets.Num();
 }
 
+void UDataprepAsset::PostDuplicate(EDuplicateMode::Type DuplicateMode)
+{
+	UDataprepAssetInterface::PostDuplicate(DuplicateMode);
+
+	// Make sure the output level does not override the one from the original
+	FText OutReason;
+	Output->SetLevelName(GetName() + "_MAP", OutReason);
+}
+
 const UDataprepActionAsset* UDataprepAsset::GetAction(int32 Index) const
 {
 	if ( ActionAssets.IsValidIndex( Index ) )
@@ -468,11 +479,13 @@ void UDataprepAsset::GetExistingParameterNamesForType(FProperty* Property, bool
 #ifndef NO_BLUEPRINT
 void UDataprepAsset::OnDataprepBlueprintChanged( UBlueprint* InBlueprint )
 {
+#if 0
 	if(InBlueprint == DataprepRecipeBP)
 	{
 		UpdateActions();
 		OnChanged.Broadcast( FDataprepAssetChangeType::RecipeModified );
 	}
+#endif
 }
 
 void UDataprepAsset::UpdateActions(bool bNotify)
diff --git a/Engine/Plugins/Enterprise/DataprepEditor/Source/DataprepCore/Private/DataprepAssetInterface.cpp b/Engine/Plugins/Enterprise/DataprepEditor/Source/DataprepCore/Private/DataprepAssetInterface.cpp
--- a/Engine/Plugins/Enterprise/DataprepEditor/Source/DataprepCore/Private/DataprepAssetInterface.cpp
+++ b/Engine/Plugins/Enterprise/DataprepEditor/Source/DataprepCore/Private/DataprepAssetInterface.cpp
@@ -91,6 +91,10 @@ UDataprepContentConsumer* UDataprepAssetInterface::SetConsumer(UClass* NewConsum
 		Output = NewObject< UDataprepContentConsumer >( this, NewConsumerClass, ConsumerName, RF_Transactional );
 		check( Output );
 
+		// Set the name of the output level
+		FText OutText;
+		Output->SetLevelName(GetName() + TEXT("_Map"), OutText);
+
 		FAssetRegistryModule::AssetCreated( Output );
 		Output->MarkPackageDirty();
 
diff --git a/Engine/Plugins/Enterprise/DataprepEditor/Source/DataprepCore/Private/DataprepContentConsumer.cpp b/Engine/Plugins/Enterprise/DataprepEditor/Source/DataprepCore/Private/DataprepContentConsumer.cpp
--- a/Engine/Plugins/Enterprise/DataprepEditor/Source/DataprepCore/Private/DataprepContentConsumer.cpp
+++ b/Engine/Plugins/Enterprise/DataprepEditor/Source/DataprepCore/Private/DataprepContentConsumer.cpp
@@ -8,7 +8,6 @@
 
 #include "EditorLevelUtils.h"
 #include "Engine/Level.h"
-#include "Engine/LevelStreamingDynamic.h"
 #include "Engine/World.h"
 #include "GameFramework/Actor.h"
 #include "HAL/FileManager.h"
@@ -18,6 +17,16 @@
 
 #define LOCTEXT_NAMESPACE "DataprepContentConsumer"
 
+const FString UDataprepContentConsumer::RelativeOutput = TEXT("ContentConsumer_RelativeOutput");
+
+const FString& UDataprepConsumerUserData::GetMarker(const FString& MarkerName) const
+{
+	static FString NullString;
+
+	const FString* ValuePtr = Markers.Find(MarkerName);
+	return ValuePtr ? *ValuePtr : NullString;
+}
+
 UDataprepContentConsumer::UDataprepContentConsumer()
 {
 	TargetContentFolder = FPaths::GetPath( GetOutermost()->GetPathName() );
@@ -119,7 +128,7 @@ FString UDataprepContentConsumer::GetTargetPackagePath() const
 	}
 
 	// If path is one level deep, make sure it ends with a '/'
-	int32 Index = -1;
+	int32 Index = INDEX_NONE;
 	TargetPackagePath.FindLastChar(L'/', Index);
 	check(Index >= 0);
 	if(Index == 0)
diff --git a/Engine/Plugins/Enterprise/DataprepEditor/Source/DataprepCore/Private/DataprepOperation.cpp b/Engine/Plugins/Enterprise/DataprepEditor/Source/DataprepCore/Private/DataprepOperation.cpp
--- a/Engine/Plugins/Enterprise/DataprepEditor/Source/DataprepCore/Private/DataprepOperation.cpp
+++ b/Engine/Plugins/Enterprise/DataprepEditor/Source/DataprepCore/Private/DataprepOperation.cpp
@@ -9,6 +9,7 @@
 #define LOCTEXT_NAMESPACE "DataprepOperation"
 
 FText FDataprepOperationCategories::ActorOperation(LOCTEXT("DataprepOperation_ActorOperationName", "On Actor"));
+FText FDataprepOperationCategories::AssetOperation(LOCTEXT("DataprepOperationAssetOperationName", "On Asset"));
 FText FDataprepOperationCategories::MeshOperation( LOCTEXT("DataprepOperation_MeshOperationName", "On Mesh"));
 FText FDataprepOperationCategories::ObjectOperation( LOCTEXT("DataprepOperation_ObjectOperationName", "On Object"));
 
diff --git a/Engine/Plugins/Enterprise/DataprepEditor/Source/DataprepCore/Public/DataprepAsset.h b/Engine/Plugins/Enterprise/DataprepEditor/Source/DataprepCore/Public/DataprepAsset.h
--- a/Engine/Plugins/Enterprise/DataprepEditor/Source/DataprepCore/Public/DataprepAsset.h
+++ b/Engine/Plugins/Enterprise/DataprepEditor/Source/DataprepCore/Public/DataprepAsset.h
@@ -43,6 +43,7 @@ public:
 	virtual void PostLoad() override;
 	virtual bool Rename(const TCHAR* NewName/* =nullptr */, UObject* NewOuter/* =nullptr */, ERenameFlags Flags/* =REN_None */) override;
 	virtual void PostEditUndo() override;
+	virtual void PostDuplicate(EDuplicateMode::Type DuplicateMode) override;
 	// End of UObject interface
 
 	// UDataprepAssetInterface interface
diff --git a/Engine/Plugins/Enterprise/DataprepEditor/Source/DataprepCore/Public/DataprepContentConsumer.h b/Engine/Plugins/Enterprise/DataprepEditor/Source/DataprepCore/Public/DataprepContentConsumer.h
--- a/Engine/Plugins/Enterprise/DataprepEditor/Source/DataprepCore/Public/DataprepContentConsumer.h
+++ b/Engine/Plugins/Enterprise/DataprepEditor/Source/DataprepCore/Public/DataprepContentConsumer.h
@@ -7,13 +7,37 @@
 #include "IDataprepLogger.h"
 #include "IDataprepProgressReporter.h"
 
+#include "Engine/AssetUserData.h"
 #include "Engine/World.h"
 
 #include "DataprepContentConsumer.generated.h"
 
+/**
+ * A UDataprepConsumerUserData is used to mark an asset or an actor to be respectively moved to a sub-folder or sub-level
+ * A DataprepContentConsumer is not obligated to honor this marker 
+ */
+UCLASS(BlueprintType, meta = (ScriptName = "DataprepConsumerUserData", DisplayName = "Dataprep Consumer User Data"))
+class DATAPREPCORE_API UDataprepConsumerUserData : public UAssetUserData
+{
+	GENERATED_BODY()
+
+public:
+	void AddMarker(const FString& MarkerName, const FString& Value)
+	{
+		Markers.FindOrAdd(MarkerName) = Value;
+	}
+
+	const FString& GetMarker(const FString& MarkerName) const;
+
+protected:
+	/** Map for Dataprep entities to store data consumable or not by a consumer */
+	UPROPERTY()
+	TMap<FString,FString> Markers;
+};
+
 struct FDataprepConsumerContext
 {
-	FDataprepConsumerContext() {}
+	FDataprepConsumerContext() : bSilentMode(false) {}
 
 	FDataprepConsumerContext& SetWorld( UWorld* InWorld )
 	{ 
@@ -60,6 +84,12 @@ struct FDataprepConsumerContext
 
 	/** Hold onto the logger that the consumer should use to log messages */
 	TSharedPtr<  IDataprepLogger > LoggerPtr;
+
+	/**
+	 * Indicates the execution of the consumer must be silent
+	 * This member is true when executed in headless mode
+	 */
+	bool bSilentMode;
 };
 
 /**
@@ -164,6 +194,9 @@ public:
 		return OnChanged;
 	}
 
+	/** Name of marker to be used by consumers to indicate specific output for assets or actors */
+	static const FString RelativeOutput;
+
 protected:
 
 	/**
diff --git a/Engine/Plugins/Enterprise/DataprepEditor/Source/DataprepCore/Public/DataprepOperation.h b/Engine/Plugins/Enterprise/DataprepEditor/Source/DataprepCore/Public/DataprepOperation.h
--- a/Engine/Plugins/Enterprise/DataprepEditor/Source/DataprepCore/Public/DataprepOperation.h
+++ b/Engine/Plugins/Enterprise/DataprepEditor/Source/DataprepCore/Public/DataprepOperation.h
@@ -58,6 +58,7 @@ class DATAPREPCORE_API FDataprepOperationCategories
 {
 public:
 	static FText ActorOperation;
+	static FText AssetOperation;
 	static FText MeshOperation;
 	static FText ObjectOperation;
 };
diff --git a/Engine/Plugins/Enterprise/DataprepEditor/Source/DataprepEditor/Private/Widgets/DataprepAssetView.cpp b/Engine/Plugins/Enterprise/DataprepEditor/Source/DataprepEditor/Private/Widgets/DataprepAssetView.cpp
--- a/Engine/Plugins/Enterprise/DataprepEditor/Source/DataprepEditor/Private/Widgets/DataprepAssetView.cpp
+++ b/Engine/Plugins/Enterprise/DataprepEditor/Source/DataprepEditor/Private/Widgets/DataprepAssetView.cpp
@@ -326,17 +326,9 @@ TSharedRef<ITableRow> SDataprepAssetView::OnGenerateRowForCategoryTree( TSharedR
 		}
 		case EDataprepCategory::Consumers:
 		{
-			ConsumerWidget = SNew( SDataprepConsumerWidget )
-				.DataprepConsumer( DataprepAssetInterfacePtr->GetConsumer() )
-				.ColumnSizeData( ColumnSizeData );
-
 			TSharedPtr< SVerticalBox > ConsumerContainer = SNew( SVerticalBox )
 				+ SVerticalBox::Slot()
-				.AutoHeight()
-				[
-					ConsumerWidget.ToSharedRef()
-				]
-				+ SVerticalBox::Slot()
+				.Padding(10.0f, 5.0f, 0.0f, 5.0f)
 				.AutoHeight()
 				[
 					SNew( SDataprepDetailsView )
@@ -637,12 +629,7 @@ void SDataprepAssetView::OnDataprepAssetChanged(FDataprepAssetChangeType ChangeT
 {
 	if( UDataprepAssetInterface* DataprepAsset = DataprepAssetInterfacePtr.Get() )
 	{
-		if(ChangeType == FDataprepAssetChangeType::ConsumerModified)
-		{
-			// Update the widget holding the consumer
-			ConsumerWidget->SetDataprepConsumer( DataprepAsset->GetConsumer() );
-		}
-		else if( ChangeType == FDataprepAssetChangeType::ProducerModified ||
+		if( ChangeType == FDataprepAssetChangeType::ProducerModified ||
 			ChangeType == FDataprepAssetChangeType::ProducerAdded ||
 			ChangeType == FDataprepAssetChangeType::ProducerRemoved )
 		{
diff --git a/Engine/Plugins/Enterprise/DataprepEditor/Source/DataprepEditor/Private/Widgets/DataprepWidgets.cpp b/Engine/Plugins/Enterprise/DataprepEditor/Source/DataprepEditor/Private/Widgets/DataprepWidgets.cpp
--- a/Engine/Plugins/Enterprise/DataprepEditor/Source/DataprepEditor/Private/Widgets/DataprepWidgets.cpp
+++ b/Engine/Plugins/Enterprise/DataprepEditor/Source/DataprepEditor/Private/Widgets/DataprepWidgets.cpp
@@ -55,45 +55,89 @@
 
 namespace DataprepWidgetUtils
 {
-	TSharedRef<SWidget> CreatePropertyWidget( TSharedPtr<SWidget> NameWidget, TSharedPtr<SWidget> ValueWidget, TSharedPtr< FDataprepDetailsViewColumnSizeData > ColumnSizeData, float Spacing, EHorizontalAlignment HAlign = EHorizontalAlignment::HAlign_Left, EVerticalAlignment VAlign = EVerticalAlignment::VAlign_Center )
+	class SCustomSplitter : public SSplitter
 	{
-		return SNew(SHorizontalBox)
-		+ SHorizontalBox::Slot()
-		.FillWidth(1.f)
-		.VAlign(VAlign_Center)
-		.HAlign(HAlign_Left)
-		.Padding(0.0f, 0.0f, 0.0f, Spacing)
-		[
-			SNew(SSplitter)
-			.Style(FEditorStyle::Get(), "DetailsView.Splitter")
-			.PhysicalSplitterHandleSize(1.0f)
-			.HitDetectionSplitterHandleSize(5.0f)
-			.ResizeMode( ESplitterResizeMode::Fill )
-			+ SSplitter::Slot()
+	public:
+		SLATE_BEGIN_ARGS(SCustomSplitter) {}
+			SLATE_ARGUMENT( TSharedPtr<SWidget>, NameWidget )
+			SLATE_ARGUMENT( TSharedPtr<SWidget>, ValueWidget )
+			SLATE_ARGUMENT( TSharedPtr< FDataprepDetailsViewColumnSizeData >, ColumnSizeData )
+		SLATE_END_ARGS();
+
+		void Construct(const FArguments& InArgs)
+		{
+			SSplitter::FArguments Args;
+
+			SSplitter::Construct(Args
+				.Style(FEditorStyle::Get(), "DetailsView.Splitter")
+				.PhysicalSplitterHandleSize(1.0f)
+				.HitDetectionSplitterHandleSize(5.0f)
+				.ResizeMode( ESplitterResizeMode::Fill ));
+
+			ColumnSizeData = InArgs._ColumnSizeData;
+
+			AddSlot()
 			.Value(ColumnSizeData->LeftColumnWidth)
 			.OnSlotResized( SSplitter::FOnSlotResized::CreateLambda( [](float InNewWidth) -> void {} ) )
 			[
-				NameWidget.ToSharedRef()
-			]
-			+ SSplitter::Slot()
-			.Value(ColumnSizeData->RightColumnWidth)
+				InArgs._NameWidget.ToSharedRef()
+			];
+
+			AddSlot()
+			.Value(InArgs._ColumnSizeData->RightColumnWidth)
 			.OnSlotResized(ColumnSizeData->OnWidthChanged)
 			[
 				SNew(SHorizontalBox)
-				.Clipping(EWidgetClipping::OnDemand)
 				+ SHorizontalBox::Slot()
 				.Padding( 5.0f, 2.5f, 2.0f, 2.5f )
-				.HAlign( HAlign )
-				.VAlign( VAlign )
+				.HAlign( HAlign_Fill )
+				.VAlign( VAlign_Fill )
+				.FillWidth(1.f)
 				[
-					// Trick to force the splitter widget to fill up the space of its parent
-					// Strongly inspired from SDetailSingleItemRow
-					SNew(DataprepWidgetUtils::SConstrainedBox)
-					[
-						ValueWidget.ToSharedRef()
-					]
+					InArgs._ValueWidget.ToSharedRef()
 				]
-			]
+			];
+		}
+
+		virtual void OnArrangeChildren(const FGeometry& AllottedGeometry, FArrangedChildren& ArrangedChildren) const override
+		{
+			if(Children.Num() > 1)
+			{
+				const SSplitter::FSlot& LeftChild = Children[0];
+
+				const FVector2D AllottedSize = AllottedGeometry.GetLocalSize();
+
+				FVector2D LocalPosition = FVector2D::ZeroVector;
+				FVector2D LocalSize(  AllottedSize.X * ColumnSizeData->LeftColumnWidth.Get(0.f), AllottedSize.Y );
+
+				ArrangedChildren.AddWidget( EVisibility::Visible, AllottedGeometry.MakeChild( LeftChild.GetWidget(), LocalPosition, LocalSize ));
+
+				const SSplitter::FSlot& RightChild = Children[1];
+
+				LocalPosition = FVector2D(LocalSize.X, 0.f);
+				LocalSize.X = AllottedSize.X * ColumnSizeData->RightColumnWidth.Get(0.f);
+
+				ArrangedChildren.AddWidget( EVisibility::Visible, AllottedGeometry.MakeChild( RightChild.GetWidget(), LocalPosition, LocalSize ));
+			}
+		}
+
+	private:
+		TSharedPtr< FDataprepDetailsViewColumnSizeData > ColumnSizeData;
+};
+
+	TSharedRef<SWidget> CreatePropertyWidget( TSharedPtr<SWidget> NameWidget, TSharedPtr<SWidget> ValueWidget, TSharedPtr< FDataprepDetailsViewColumnSizeData > ColumnSizeData, float Spacing)
+	{
+		return SNew(SHorizontalBox)
+		+ SHorizontalBox::Slot()
+		.FillWidth(1.f)
+		.VAlign(VAlign_Fill)
+		.HAlign(HAlign_Fill)
+		.Padding(0.0f, 0.0f, 0.0f, Spacing)
+		[
+			SNew(SCustomSplitter)
+			.NameWidget(NameWidget)
+			.ValueWidget(ValueWidget)
+			.ColumnSizeData(ColumnSizeData)
 		];
 	}
 }
@@ -194,358 +238,7 @@ void SDataprepCategoryWidget::Construct( const FArguments& InArgs, TSharedRef< S
 	);
 }
 
-void SDataprepConsumerWidget::OnLevelNameChanged( const FText &NewLevelName, ETextCommit::Type CommitType )
-{
-	const FScopedTransaction Transaction( LOCTEXT("Consumer_SetLevelName", "Set Level Name") );
-	FText OutReason;
-	if( !DataprepConsumerPtr->SetLevelName( NewLevelName.ToString(), OutReason ) )
-	{
-		UE_LOG( LogDataprepEditor, Error, TEXT("%s"), *OutReason.ToString() );
-		LevelTextBox->SetText( FText::FromString( DataprepConsumerPtr->GetLevelName() ) );
-	}
-}
-
-void SDataprepConsumerWidget::OnTextCommitted( const FText& NewText, ETextCommit::Type CommitType)
-{
-	if( UDataprepContentConsumer* DataprepConsumer = DataprepConsumerPtr.Get() )
-	{
-		FString NewContentFolder( NewText.ToString() );
-
-		// Replace /Content/ with /Game/ since /Content is only used for display 
-		if( NewContentFolder.StartsWith( TEXT("/Content") ) )
-		{
-			NewContentFolder = NewContentFolder.Replace( TEXT( "/Content" ), TEXT( "/Game" ) );
-		}
-		
-		// Remove ending '/' if applicable
-		if( !NewContentFolder.IsEmpty() && NewContentFolder[ NewContentFolder.Len()-1 ] == TEXT('/') )
-	{
-			NewContentFolder[ NewContentFolder.Len()-1 ] = 0;
-			NewContentFolder = NewContentFolder.LeftChop(1);
-	}
-
-		const FScopedTransaction Transaction( LOCTEXT("Consumer_SetTargetContentFolder", "Set Target Content Folder") );
-
-		FText ErrorReason;
-		if( !DataprepConsumer->SetTargetContentFolder( NewContentFolder, ErrorReason ) )
-	{
-			UE_LOG( LogDataprepEditor, Error, TEXT("%s"), *ErrorReason.ToString() );
-			UpdateContentFolderText();
-		}
-	}
-}
-
-void SDataprepConsumerWidget::OnBrowseContentFolder()
-{
-	FString Path = DataprepConsumerPtr->GetTargetContentFolder();
-	if( Path.IsEmpty() )
-	{
-		Path = FPaths::GetPath( DataprepConsumerPtr->GetOutermost()->GetPathName() );
-	}
-	Path += TEXT("/"); // Trailing '/' is needed to set the default path
-
-	//Ask the user for the root path where they want any content to be placed
-	if( UDataprepContentConsumer* DataprepConsumer = DataprepConsumerPtr.Get() )
-	{
-		TSharedRef<SDlgPickPath> PickContentPathDlg =
-			SNew(SDlgPickPath)
-			.Title(LOCTEXT("DataprepSlateHelper_ChooseImportRootContentPath", "Choose Location for importing the Datasmith content"))
-			.DefaultPath(FText::FromString(Path));
-
-		if ( PickContentPathDlg->ShowModal() == EAppReturnType::Ok )
-		{
-			const FScopedTransaction Transaction( LOCTEXT("Consumer_SetTargetContentFolder", "Set Target Content Folder") );
-
-			FText ErrorReason;
-			if( DataprepConsumer->SetTargetContentFolder( PickContentPathDlg->GetPath().ToString(), ErrorReason ) )
-			{
-			UpdateContentFolderText();
-		}
-			else
-			{
-				UE_LOG( LogDataprepEditor, Error, TEXT("%s"), *ErrorReason.ToString() );
-			}
-		}
-	}
-}
-
-void SDataprepConsumerWidget::SetDataprepConsumer(UDataprepContentConsumer* InDataprepConsumer)
-{
-	if(InDataprepConsumer == nullptr)
-	{
-		return;
-	}
-
-	if(DataprepConsumerPtr != InDataprepConsumer)
-	{
-		if(DataprepConsumerPtr.IsValid())
-		{
-			DataprepConsumerPtr->GetOnChanged().Remove( OnConsumerChangedHandle );
-		}
-
-		DataprepConsumerPtr = InDataprepConsumer;
-
-		OnConsumerChangedHandle = InDataprepConsumer->GetOnChanged().AddSP( this, &SDataprepConsumerWidget::OnConsumerChanged );
-
-		OnConsumerChanged();
-	}
-}
-
-void SDataprepConsumerWidget::Construct(const FArguments& InArgs )
-{
-	if (InArgs._ColumnSizeData.IsValid())
-	{
-		ColumnSizeData = InArgs._ColumnSizeData;
-	}
-	else
-	{
-		ColumnWidth = 0.7f;
-		ColumnSizeData = MakeShared<FDataprepDetailsViewColumnSizeData>();
-		ColumnSizeData->LeftColumnWidth = TAttribute<float>(this, &SDataprepConsumerWidget::OnGetLeftColumnWidth);
-		ColumnSizeData->RightColumnWidth = TAttribute<float>(this, &SDataprepConsumerWidget::OnGetRightColumnWidth);
-		ColumnSizeData->OnWidthChanged = SSplitter::FOnSlotResized::CreateSP(this, &SDataprepConsumerWidget::OnSetColumnWidth);
-	}
-
-	TSharedPtr<SWidget> InnerWidget = InArgs._DataprepConsumer ? BuildWidget() : BuildNullWidget();
-
-	ChildSlot
-	[
-		InnerWidget.ToSharedRef()
-	];
-
-	if(InArgs._DataprepConsumer)
-	{
-		SetDataprepConsumer( InArgs._DataprepConsumer );
-	}
-}
-
-SDataprepConsumerWidget::~SDataprepConsumerWidget()
-{
-	if ( UDataprepContentConsumer* DataprepConsumer = DataprepConsumerPtr.Get() )
-	{
-		DataprepConsumer->GetOnChanged().Remove( OnConsumerChangedHandle );
-	}
-}
-
-TSharedRef<SWidget> SDataprepConsumerWidget::BuildWidget()
-{
-	TSharedRef<SWidget> BrowseButton = PropertyCustomizationHelpers::MakeBrowseButton( FSimpleDelegate::CreateSP( this, &SDataprepConsumerWidget::OnBrowseContentFolder ) );
-
-	TSharedRef<SWidget> Widget = SNew(SBorder)
-	.BorderImage( FCoreStyle::Get().GetBrush("NoBrush") )
-	[
-		SNew(SGridPanel)
-		.FillColumn(0, 1.0f)
-		+ SGridPanel::Slot(0, 0)
-		.Padding(10.0f, 5.0f, 0.0f, 5.0f)
-		[
-			SNew(SHorizontalBox)
-			+ SHorizontalBox::Slot()
-			.FillWidth(1.f)
-			.VAlign(VAlign_Center)
-			.HAlign(HAlign_Left)
-			[
-				SNew(SSplitter)
-				.Style(FEditorStyle::Get(), "DetailsView.Splitter")
-				.PhysicalSplitterHandleSize(1.0f)
-				.HitDetectionSplitterHandleSize(5.0f)
-				.ResizeMode(ESplitterResizeMode::Fill)
-				+ SSplitter::Slot()
-				.Value(ColumnSizeData->LeftColumnWidth)
-				.OnSlotResized(SSplitter::FOnSlotResized::CreateSP(this, &SDataprepConsumerWidget::OnLeftColumnResized))
-				[
-					SNew(SHorizontalBox)
-					.Clipping(EWidgetClipping::OnDemand)
-					+ SHorizontalBox::Slot()
-					.HAlign(HAlign_Left)
-					.VAlign(VAlign_Center)
-					[
-						SNew(STextBlock)
-						.Text(LOCTEXT("DataprepSlateHelper_ContentFolderLabel", "Folder"))
-						.Font(IDetailLayoutBuilder::GetDetailFont())
-					]
-				]
-				+ SSplitter::Slot()
-				.Value(ColumnSizeData->RightColumnWidth)
-				.OnSlotResized(ColumnSizeData->OnWidthChanged)
-				[
-					SNew(SHorizontalBox)
-					.Clipping(EWidgetClipping::OnDemand)
-					+ SHorizontalBox::Slot()
-					.FillWidth(1.0f)
-					.Padding(5.0f, 2.5f, 2.0f, 2.5f)
-					[
-						// Trick to force the splitter widget to fill up the space of its parent
-						// Strongly inspired from SDetailSingleItemRow
-						SNew(DataprepWidgetUtils::SConstrainedBox)
-						[
-							SAssignNew(ContentFolderTextBox, SEditableTextBox)
-							.Font(IDetailLayoutBuilder::GetDetailFont())
-							.HintText(LOCTEXT("DataprepSlateHelper_ContentFolderHintText", "Set the content folder to save in"))
-							.IsReadOnly(false)
-							.OnTextCommitted(FOnTextCommitted::CreateSP(this, &SDataprepConsumerWidget::OnTextCommitted))
-						]
-					]
-				]
-			]
-		]
-		+ SGridPanel::Slot(1, 0)
-		.Padding(5.0f, 5.0f, 0.0f, 5.0f)
-		[
-			SNew(SHorizontalBox)
-			+ SHorizontalBox::Slot()
-			.VAlign(VAlign_Center)
-			.HAlign(HAlign_Right)
-			.AutoWidth()
-			[
-				BrowseButton
-			]
-		]
-		// This column is required to align consumer section with producers' section which has 3 columns.
-		+ SGridPanel::Slot(2, 0)
-		.Padding(5.0f, 5.0f, 0.0f, 5.0f)
-		[
-			SNew(SHorizontalBox)
-			+ SHorizontalBox::Slot()
-			.VAlign(VAlign_Center)
-			.HAlign(HAlign_Right)
-			.AutoWidth()
-			[
-				SNew(SButton)
-				.IsFocusable(false)
-				.Visibility(EVisibility::Hidden)
-				.IsEnabled(false)
-				.VAlign(VAlign_Top)
-				.Content()
-				[
-					SNew(STextBlock)
-					.Font(FDataprepEditorUtils::GetGlyphFont())
-					.ColorAndOpacity(FLinearColor::Transparent)
-					.Text(FEditorFontGlyphs::Exclamation_Triangle)
-				]
-			]
-		]
-		+ SGridPanel::Slot(0, 1)
-		.Padding(10.0f, 5.0f, 0.0f, 5.0f)
-		[
-			SNew(SHorizontalBox)
-			+ SHorizontalBox::Slot()
-			.FillWidth(1.f)
-			.VAlign(VAlign_Center)
-			.HAlign(HAlign_Left)
-			[
-				SNew(SSplitter)
-				.Style(FEditorStyle::Get(), "DetailsView.Splitter")
-				.PhysicalSplitterHandleSize(1.0f)
-				.HitDetectionSplitterHandleSize(5.0f)
-				.ResizeMode(ESplitterResizeMode::Fill)
-				+ SSplitter::Slot()
-				.Value(ColumnSizeData->LeftColumnWidth)
-				.OnSlotResized(SSplitter::FOnSlotResized::CreateSP(this, &SDataprepConsumerWidget::OnLeftColumnResized))
-				[
-					SNew(SHorizontalBox)
-					.Clipping(EWidgetClipping::OnDemand)
-					+ SHorizontalBox::Slot()
-					.HAlign(HAlign_Left)
-					.VAlign(VAlign_Center)
-					[
-						SNew(STextBlock)
-						.Text(LOCTEXT("DataprepSlateHelper_LevelNameLabel", "Sub-Level"))
-						.Font(IDetailLayoutBuilder::GetDetailFont())
-					]
-				]
-				+ SSplitter::Slot()
-				.Value(ColumnSizeData->RightColumnWidth)
-				.OnSlotResized(ColumnSizeData->OnWidthChanged)
-				[
-					SNew(SHorizontalBox)
-					.Clipping(EWidgetClipping::OnDemand)
-					+ SHorizontalBox::Slot()
-					.FillWidth(1.0f)
-					.Padding(5.0f, 2.5f, 2.0f, 2.5f)
-					[
-						// Trick to force the splitter widget to fill up the space of its parent
-						// Strongly inspired from SDetailSingleItemRow
-						SNew(DataprepWidgetUtils::SConstrainedBox)
-						[
-							SAssignNew(LevelTextBox, SEditableTextBox)
-							.Font(IDetailLayoutBuilder::GetDetailFont())
-							.HintText(LOCTEXT("DataprepSlateHelper_LevelNameHintText", "Current will be used"))
-							.OnTextCommitted(this, &SDataprepConsumerWidget::OnLevelNameChanged)
-						]
-					]
-				]
-			]
-		]
-	];
-
-	return Widget;
-}
-
-TSharedRef<SWidget> SDataprepConsumerWidget::BuildNullWidget()
-{
-	return SNew(SBorder)
-	.BorderImage( FCoreStyle::Get().GetBrush("ToolPanel.GroupBorder") )
-	[
-		SNew(SVerticalBox)
-		+ SVerticalBox::Slot()
-		.Padding(5.0f)
-		[
-			SNew(SHorizontalBox)
-			+ SHorizontalBox::Slot()
-			.HAlign(HAlign_Left)
-			.VAlign(VAlign_Center)
-			.FillWidth(1.0f)
-			[
-				SNew(STextBlock)
-				.Font( IDetailLayoutBuilder::GetDetailFont() )
-				.Text(LOCTEXT("DataprepSlateHelper_Error_InvalidConsumer", "Error: Not a valid consumer"))
-				.Margin(FMargin( 5.0f, 5.0f, 0.0f, 0.0f ) )
-				.ColorAndOpacity( FLinearColor(1,0,0,1) )
-			]
-		]
-	];
-}
-
-void SDataprepConsumerWidget::UpdateContentFolderText()
-{
-	if(UDataprepContentConsumer* Consumer = DataprepConsumerPtr.Get())
-	{
-		FString TargetContentFolder( Consumer->GetTargetContentFolder() );
-
-		if( TargetContentFolder.IsEmpty() )
-		{
-			TargetContentFolder = TEXT("/Content");
-		}
-		else if( TargetContentFolder.StartsWith( TEXT( "/Game" ) ) )
-		{
-			TargetContentFolder = TargetContentFolder.Replace( TEXT( "/Game" ), TEXT( "/Content" ) );
-		}
-
-		ContentFolderTextBox->SetText( FText::FromString( TargetContentFolder + TEXT( "/" ) ) );
-	}
-	else
-	{
-		ContentFolderTextBox->SetText( TAttribute<FText>() );
-	}
-}
-
-void SDataprepConsumerWidget::OnConsumerChanged()
-{
-	if(DataprepConsumerPtr.IsValid())
-	{
-		UpdateContentFolderText();
-		LevelTextBox->SetText( FText::FromString( DataprepConsumerPtr->GetLevelName() ) );
-	}
-	else
-	{
-		ContentFolderTextBox->SetText( TAttribute<FText>() );
-		LevelTextBox->SetText( TAttribute<FText>() );
-	}
-}
-
-
-void SDataprepDetailsView::CreateDefaultWidget( int32 Index, TSharedPtr< SWidget >& NameWidget, TSharedPtr< SWidget >& ValueWidget, float LeftPadding, EHorizontalAlignment HAlign, EVerticalAlignment VAlign, const FDataprepParameterizationContext& ParameterizationContext)
+void SDataprepDetailsView::CreateDefaultWidget( int32 Index, TSharedPtr< SWidget >& NameWidget, TSharedPtr< SWidget >& ValueWidget, float LeftPadding, const FDataprepParameterizationContext& ParameterizationContext)
 {
 	TSharedRef<SHorizontalBox> NameColumn = SNew(SHorizontalBox)
 		.Clipping(EWidgetClipping::OnDemand);
@@ -597,7 +290,7 @@ void SDataprepDetailsView::CreateDefaultWidget( int32 Index, TSharedPtr< SWidget
 		SNew(SDataprepContextMenuOverride)
 		.OnContextMenuOpening(OnContextMenuOpening)
 		[
-			DataprepWidgetUtils::CreatePropertyWidget( NameColumn, ValueWidget, ColumnSizeData, Spacing, HAlign, VAlign )
+			DataprepWidgetUtils::CreatePropertyWidget( NameColumn, ValueWidget, ColumnSizeData, Spacing )
 		]
 		];
 
@@ -793,7 +486,7 @@ void SDataprepDetailsView::AddWidgets( const TArray< TSharedRef< IDetailTreeNode
 				TSharedPtr< SWidget > ValueWidget;
 				DetailPropertyRow->GetDefaultWidgets( NameWidget, ValueWidget, Row, true );
 
-				CreateDefaultWidget( Index, NameWidget, ValueWidget, LeftPadding, Row.ValueWidget.HorizontalAlignment, Row.ValueWidget.VerticalAlignment, CurrentParameterizationContext );
+				CreateDefaultWidget( Index, NameWidget, ValueWidget, LeftPadding, CurrentParameterizationContext );
 				Index++;
 
 				TArray< TSharedRef< IDetailTreeNode > > Children;
@@ -834,7 +527,7 @@ void SDataprepDetailsView::AddWidgets( const TArray< TSharedRef< IDetailTreeNode
 
 			if( NameWidget.IsValid() && ValueWidget.IsValid() )
 			{
-				CreateDefaultWidget( Index, NameWidget, ValueWidget, LeftPadding, HAlign, VAlign, CurrentParameterizationContext );
+				CreateDefaultWidget( Index, NameWidget, ValueWidget, LeftPadding, CurrentParameterizationContext );
 				Index++;
 
 				bool bDisplayChildren = true;
@@ -1073,7 +766,7 @@ void SDataprepInstanceParentWidget::Construct(const FArguments& InArgs)
 	}
 	else
 	{
-		ColumnWidth = 0.7f;
+		ColumnWidth = 0.5f;
 		ColumnSizeData = MakeShared<FDataprepDetailsViewColumnSizeData>();
 		ColumnSizeData->LeftColumnWidth = TAttribute<float>(this, &SDataprepInstanceParentWidget::OnGetLeftColumnWidth);
 		ColumnSizeData->RightColumnWidth = TAttribute<float>(this, &SDataprepInstanceParentWidget::OnGetRightColumnWidth);
diff --git a/Engine/Plugins/Enterprise/DataprepEditor/Source/DataprepEditor/Private/Widgets/DataprepWidgets.h b/Engine/Plugins/Enterprise/DataprepEditor/Source/DataprepEditor/Private/Widgets/DataprepWidgets.h
--- a/Engine/Plugins/Enterprise/DataprepEditor/Source/DataprepEditor/Private/Widgets/DataprepWidgets.h
+++ b/Engine/Plugins/Enterprise/DataprepEditor/Source/DataprepEditor/Private/Widgets/DataprepWidgets.h
@@ -92,67 +92,6 @@ private:
 	TSharedPtr< SWidget > CategoryContent;
 };
 
-class SDataprepConsumerWidget : public SCompoundWidget
-{
-public:
-	SLATE_BEGIN_ARGS(SDataprepConsumerWidget)
-		: _DataprepConsumer(nullptr)
-	{}
-
-	SLATE_ARGUMENT(TSharedPtr< FDataprepDetailsViewColumnSizeData >, ColumnSizeData)
-	SLATE_ARGUMENT(UDataprepContentConsumer*, DataprepConsumer)
-	SLATE_END_ARGS()
-
-	void Construct( const FArguments& InArgs );
-
-	virtual ~SDataprepConsumerWidget();
-
-	/** Update DataprepConsumerPtr and content folder's and level name's text boxes */
-	void SetDataprepConsumer( UDataprepContentConsumer* DataprepConsumer );
-
-private:
-	/** Create complete widget from valid UDataprepContentConsumer object */
-	TSharedRef<SWidget> BuildWidget();
-	/** Create empty widget from invalid UDataprepContentConsumer object */
-	TSharedRef<SWidget> BuildNullWidget();
-	/** Callback when content of level name's text box has changed */
-	void OnLevelNameChanged( const FText& NewLevelName, ETextCommit::Type CommitType );
-	/** Callback when browser's button is clicked */
-	void OnBrowseContentFolder();
-	/** Callback when content of content folder's text box has changed */
-	void OnTextCommitted( const FText&, ETextCommit::Type );
-
-	/** Update content folder's and level name's text boxes */
-	void OnConsumerChanged();
-	/** Update content folder's text boxes */
-	void UpdateContentFolderText();
-
-	/** Callbacks to update splitter */
-	void OnLeftColumnResized(float InNewWidth)
-	{
-		// This has to be bound or the splitter will take it upon itself to determine the size
-		// We do nothing here because it is handled by the column size data
-	}
-
-	float OnGetLeftColumnWidth() const { return 1.0f - ColumnWidth; }
-	float OnGetRightColumnWidth() const { return ColumnWidth; }
-	void OnSetColumnWidth(float InWidth) { ColumnWidth = InWidth; }
-
-private:
-	/** Weak pointer on edited consumer */
-	TWeakObjectPtr<UDataprepContentConsumer> DataprepConsumerPtr;
-	/** Content folder's text box */
-	TSharedPtr< SEditableTextBox > ContentFolderTextBox;
-	/** Level name's text box */
-	TSharedPtr< SEditableTextBox > LevelTextBox;
-	/** Helps sync column resizing with another part of the UI (producers widget) */
-	TSharedPtr< FDataprepDetailsViewColumnSizeData > ColumnSizeData;
-	/** Relative width to control splitters. */
-	float ColumnWidth;
-
-	FDelegateHandle OnConsumerChangedHandle;
-};
-
 class SDataprepDetailsView : public SCompoundWidget, public FGCObject
 {
 	using Super = SCompoundWidget;
@@ -214,7 +153,7 @@ private:
 	 * @param LeftPadding				The Padding on the left of the row
 	 * @param ParameterizationContext	Parameterization context of the associated property
 	 */
-	void CreateDefaultWidget(int32 Index, TSharedPtr< SWidget >& NameWidget, TSharedPtr< SWidget >& ValueWidget, float LeftPadding, EHorizontalAlignment HAlign, EVerticalAlignment VAlign, const FDataprepParameterizationContext& ParameterizationContext);
+	void CreateDefaultWidget(int32 Index, TSharedPtr< SWidget >& NameWidget, TSharedPtr< SWidget >& ValueWidget, float LeftPadding, const FDataprepParameterizationContext& ParameterizationContext);
 
 	/** Callback used by all splitters in the details view, so that they move in sync */
 	void OnLeftColumnResized(float InNewWidth)
@@ -225,7 +164,7 @@ private:
 
 	float OnGetLeftColumnWidth() const { return 1.0f - ColumnWidth; }
 	float OnGetRightColumnWidth() const { return ColumnWidth; }
-	void OnSetColumnWidth(float InWidth) { ColumnWidth = InWidth; }
+	void OnSetColumnWidth(float InWidth) { ColumnWidth = InWidth < 0.3f ? 0.3f : InWidth; }
 
 	/** Callback to track property changes on array properties */
 	void OnPropertyChanged( const struct FPropertyChangedEvent& InEvent );
diff --git a/Engine/Plugins/Enterprise/DataprepEditor/Source/DataprepLibraries/Private/DataprepOperations.cpp b/Engine/Plugins/Enterprise/DataprepEditor/Source/DataprepLibraries/Private/DataprepOperations.cpp
--- a/Engine/Plugins/Enterprise/DataprepEditor/Source/DataprepLibraries/Private/DataprepOperations.cpp
+++ b/Engine/Plugins/Enterprise/DataprepEditor/Source/DataprepLibraries/Private/DataprepOperations.cpp
@@ -383,4 +383,23 @@ void UDataprepFlipFacesOperation::OnExecution_Implementation(const FDataprepCont
 	UStaticMesh::BatchBuild(StaticMeshes.Array());
 }
 
+void UDataprepSetOutputLevel::OnExecution_Implementation(const FDataprepContext& InContext)
+{
+#ifdef LOG_TIME
+	DataprepOperationTime::FTimeLogger TimeLogger(TEXT("RandomizeTransform"), [&](FText Text) { this->LogInfo(Text); });
+#endif
+
+	UDataprepOperationsLibrary::SetSubOuputLevel(InContext.Objects, LevelName);
+}
+
+void UDataprepSetOutputFolder::OnExecution_Implementation(const FDataprepContext& InContext)
+{
+#ifdef LOG_TIME
+	DataprepOperationTime::FTimeLogger TimeLogger(TEXT("RandomizeTransform"), [&](FText Text) { this->LogInfo(Text); });
+#endif
+
+	UDataprepOperationsLibrary::SetSubOuputFolder(InContext.Objects, FolderName);
+}
+
+
 #undef LOCTEXT_NAMESPACE
diff --git a/Engine/Plugins/Enterprise/DataprepEditor/Source/DataprepLibraries/Private/DataprepOperations.h b/Engine/Plugins/Enterprise/DataprepEditor/Source/DataprepLibraries/Private/DataprepOperations.h
--- a/Engine/Plugins/Enterprise/DataprepEditor/Source/DataprepLibraries/Private/DataprepOperations.h
+++ b/Engine/Plugins/Enterprise/DataprepEditor/Source/DataprepLibraries/Private/DataprepOperations.h
@@ -448,3 +448,57 @@ protected:
 	virtual void OnExecution_Implementation(const FDataprepContext& InContext) override;
 	//~ End UDataprepOperation Interface
 };
+
+UCLASS(Experimental, Category = ActorOperation, Meta = (DisplayName="Output to Level", ToolTip = "For each actor to process, set the sub-level to move it to.\nThe sub-level is relative to the level specified to the Dataprep consumer.") )
+class UDataprepSetOutputLevel : public UDataprepOperation
+{
+	GENERATED_BODY()
+
+	UDataprepSetOutputLevel()
+	: LevelName(TEXT("MySubLevel"))
+	{
+	}
+
+public:
+	// If true, the screen sizes at which LODs swap are computed automatically
+	UPROPERTY(EditAnywhere, BlueprintReadWrite, Category = ActorOperation, meta = (DisplayName = "Level Name", ToolTip = "Name of the sub level the actors to be saved to"))
+	FString LevelName;
+
+	//~ Begin UDataprepOperation Interface
+public:
+	virtual FText GetCategory_Implementation() const override
+	{
+		return FDataprepOperationCategories::ActorOperation;
+	}
+
+protected:
+	virtual void OnExecution_Implementation(const FDataprepContext& InContext) override;
+	//~ End UDataprepOperation Interface
+};
+
+UCLASS(Experimental, Category = AssetOperation, Meta = (DisplayName="Output to Folder", ToolTip = "For each asset to process, set the sub-folder to save it to.\nThe sub-folder is relative to the folder specified to the Dataprep consumer.") )
+class UDataprepSetOutputFolder : public UDataprepOperation
+{
+	GENERATED_BODY()
+
+	UDataprepSetOutputFolder()
+	: FolderName(TEXT("MySubFolder"))
+	{
+	}
+
+public:
+	// If true, the screen sizes at which LODs swap are computed automatically
+	UPROPERTY(EditAnywhere, BlueprintReadWrite, Category = AssetOperation, meta = (DisplayName = "Folder Name", ToolTip = "Name of the sub folder the assets to be saved to"))
+	FString FolderName;
+
+	//~ Begin UDataprepOperation Interface
+public:
+	virtual FText GetCategory_Implementation() const override
+	{
+		return FDataprepOperationCategories::AssetOperation;
+	}
+
+protected:
+	virtual void OnExecution_Implementation(const FDataprepContext& InContext) override;
+	//~ End UDataprepOperation Interface
+};
\ No newline at end of file
diff --git a/Engine/Plugins/Enterprise/DataprepEditor/Source/DataprepLibraries/Private/DataprepOperationsLibrary.cpp b/Engine/Plugins/Enterprise/DataprepEditor/Source/DataprepLibraries/Private/DataprepOperationsLibrary.cpp
--- a/Engine/Plugins/Enterprise/DataprepEditor/Source/DataprepLibraries/Private/DataprepOperationsLibrary.cpp
+++ b/Engine/Plugins/Enterprise/DataprepEditor/Source/DataprepLibraries/Private/DataprepOperationsLibrary.cpp
@@ -3,6 +3,7 @@
 #include "DataprepOperationsLibrary.h"
 
 #include "DataprepCoreUtils.h"
+#include "DataprepContentConsumer.h"
 #include "DatasmithAssetUserData.h"
 
 #include "ActorEditorUtils.h"
@@ -911,4 +912,70 @@ void UDataprepOperationsLibrary::FlipFaces(const TSet< UStaticMesh* >& StaticMes
 	}
 }
 
+void UDataprepOperationsLibrary::SetSubOuputLevel(const TArray<UObject*>& SelectedObjects, const FString& SubLevelName)
+{
+	if(SubLevelName.IsEmpty())
+	{
+		return;
+	}
+
+	for (UObject* Object : SelectedObjects)
+	{
+		if (AActor* Actor = Cast< AActor >(Object))
+		{
+			if (USceneComponent* RootComponent = Actor->GetRootComponent())
+			{
+				if ( RootComponent->GetClass()->ImplementsInterface(UInterface_AssetUserData::StaticClass()) )
+				{
+					if ( IInterface_AssetUserData* AssetUserDataInterface = Cast< IInterface_AssetUserData >( RootComponent ) )
+					{
+						UDataprepConsumerUserData* DataprepContentUserData = AssetUserDataInterface->GetAssetUserData< UDataprepConsumerUserData >();
+
+						if ( !DataprepContentUserData )
+						{
+							EObjectFlags Flags = RF_Public;
+							DataprepContentUserData = NewObject< UDataprepConsumerUserData >( RootComponent, NAME_None, Flags );
+							AssetUserDataInterface->AddAssetUserData( DataprepContentUserData );
+						}
+
+						DataprepContentUserData->AddMarker(UDataprepContentConsumer::RelativeOutput, SubLevelName);
+					}
+				}
+			}
+		}
+	}
+}
+
+void UDataprepOperationsLibrary::SetSubOuputFolder(const TArray<UObject*>& SelectedObjects, const FString& SubFolderName)
+{
+	if(SubFolderName.IsEmpty())
+	{
+		return;
+	}
+
+	for (UObject* Object : SelectedObjects)
+	{
+		const bool bValidObject = Object->HasAnyFlags(RF_Public)
+		&& !Object->IsPendingKill()
+		&& Object->GetClass()->ImplementsInterface(UInterface_AssetUserData::StaticClass());
+
+		if (bValidObject)
+		{
+			if ( IInterface_AssetUserData* AssetUserDataInterface = Cast< IInterface_AssetUserData >( Object ) )
+			{
+				UDataprepConsumerUserData* DataprepContentUserData = AssetUserDataInterface->GetAssetUserData< UDataprepConsumerUserData >();
+
+				if ( !DataprepContentUserData )
+				{
+					EObjectFlags Flags = RF_Public;
+					DataprepContentUserData = NewObject< UDataprepConsumerUserData >( Object, NAME_None, Flags );
+					AssetUserDataInterface->AddAssetUserData( DataprepContentUserData );
+				}
+
+				DataprepContentUserData->AddMarker(UDataprepContentConsumer::RelativeOutput, SubFolderName);
+			}
+		}
+	}
+}
+
 #undef LOCTEXT_NAMESPACE
diff --git a/Engine/Plugins/Enterprise/DataprepEditor/Source/DataprepLibraries/Private/DataprepOperationsLibrary.h b/Engine/Plugins/Enterprise/DataprepEditor/Source/DataprepLibraries/Private/DataprepOperationsLibrary.h
--- a/Engine/Plugins/Enterprise/DataprepEditor/Source/DataprepLibraries/Private/DataprepOperationsLibrary.h
+++ b/Engine/Plugins/Enterprise/DataprepEditor/Source/DataprepLibraries/Private/DataprepOperationsLibrary.h
@@ -319,6 +319,24 @@ public:
 	UFUNCTION(BlueprintCallable, Category = "Dataprep | Operation")
 	static void FlipFaces( const TSet< UStaticMesh* >& StaticMeshes );
 
+	/**
+	 * Add/Edit UDataprepConsumerUserData with the requested name for the sub-level
+	 * @param SelectedObjects:	Objects to consider
+	 * @param SubLevelName:	Name of the sub-level
+	 * @note - This operation only applies on actors
+	 */
+	UFUNCTION(BlueprintCallable, Category = "Dataprep | Operation")
+	static void SetSubOuputLevel(const TArray<UObject*>& SelectedObjects, const FString& SubLevelName);
+
+	/**
+	 * Add/Edit UDataprepConsumerUserData with the requested name for the sub-folder
+	 * @param SelectedObjects:	Objects to consider
+	 * @param SubFolderName:	Name of the sub-folder
+	 * @note - This operation only applies on assets
+	 */
+	UFUNCTION(BlueprintCallable, Category = "Dataprep | Operation")
+	static void SetSubOuputFolder(const TArray<UObject*>& SelectedObjects, const FString& SubFolderName);
+
 private:
 	static void SubstituteMaterial(const TArray<UObject*>& SelectedObjects, const FString& MaterialSearch, EEditorScriptingStringMatchType StringMatch, const TArray<UMaterialInterface*>& MaterialList, UMaterialInterface* MaterialSubstitute);
 	static void SubstituteMesh(const TArray<UObject*>& SelectedObjects, const FString& MeshSearch, EEditorScriptingStringMatchType StringMatch, const TArray<UStaticMesh*>& MeshList, UStaticMesh* MeshSubstitute);
diff --git a/Engine/Plugins/Enterprise/DataprepEditor/Source/DataprepLibraries/Private/DataprepSelectionTransforms.cpp b/Engine/Plugins/Enterprise/DataprepEditor/Source/DataprepLibraries/Private/DataprepSelectionTransforms.cpp
--- a/Engine/Plugins/Enterprise/DataprepEditor/Source/DataprepLibraries/Private/DataprepSelectionTransforms.cpp
+++ b/Engine/Plugins/Enterprise/DataprepEditor/Source/DataprepLibraries/Private/DataprepSelectionTransforms.cpp
@@ -85,4 +85,67 @@ void UDataprepReferenceSelectionTransform::OnExecution_Implementation(const TArr
 	OutObjects.Append(Assets.Array());
 }
 
+void UDataprepHierarchySelectionTransform::OnExecution_Implementation(const TArray<UObject*>& InObjects, TArray<UObject*>& OutObjects)
+{
+	TArray<AActor*> ActorsToVisit;
+
+	for (UObject* Object : InObjects)
+	{
+		if (!ensure(Object) || Object->IsPendingKill())
+		{
+			continue;
+		}
+
+		if (AActor* Actor = Cast< AActor >(Object))
+		{
+			TArray<AActor*> Children;
+			Actor->GetAttachedActors( Children );
+
+			ActorsToVisit.Append( Children );
+		}
+	}
+
+	TSet<UObject*> NewSelection;
+
+	while ( ActorsToVisit.Num() > 0)
+	{
+		AActor* VisitedActor = ActorsToVisit.Pop();
+		if (VisitedActor == nullptr)
+		{
+			continue;
+		}
+
+		NewSelection.Add(VisitedActor);
+
+		if(SelectionPolicy == EDataprepHierarchySelectionPolicy::AllDescendants)
+		{
+			// Continue with children
+			TArray<AActor*> Children;
+			VisitedActor->GetAttachedActors( Children );
+
+			ActorsToVisit.Append( Children );
+		}
+	}
+
+	OutObjects.Append(NewSelection.Array());
+
+	if (bOutputCanIncludeInput)
+	{
+		OutObjects.Reserve( OutObjects.Num() + InObjects.Num());
+
+		for (UObject* Object : InObjects)
+		{
+			if (!ensure(Object) || Object->IsPendingKill())
+			{
+				continue;
+			}
+
+			if (AActor* Actor = Cast< AActor >(Object))
+			{
+				OutObjects.Add(Object);
+			}
+		}
+	}
+}
+
 #undef LOCTEXT_NAMESPACE
diff --git a/Engine/Plugins/Enterprise/DataprepEditor/Source/DataprepLibraries/Private/DataprepSelectionTransforms.h b/Engine/Plugins/Enterprise/DataprepEditor/Source/DataprepLibraries/Private/DataprepSelectionTransforms.h
--- a/Engine/Plugins/Enterprise/DataprepEditor/Source/DataprepLibraries/Private/DataprepSelectionTransforms.h
+++ b/Engine/Plugins/Enterprise/DataprepEditor/Source/DataprepLibraries/Private/DataprepSelectionTransforms.h
@@ -12,6 +12,16 @@
 
 #include "DataprepSelectionTransforms.generated.h"
 
+UENUM()
+enum class EDataprepHierarchySelectionPolicy : uint8
+{
+	/** Select immediate children of the selected objects */
+	ImmediateChildren,
+
+	/** Select all descendants of the selected objects */
+	AllDescendants,
+};
+
 UCLASS(Category = SelectionTransform, Meta = (DisplayName="Reference Selection Transform", ToolTip = "Return all the assets used/referenced by the selected objects") )
 class UDataprepReferenceSelectionTransform : public UDataprepSelectionTransform
 {
@@ -20,3 +30,19 @@ class UDataprepReferenceSelectionTransform : public UDataprepSelectionTransform
 protected:
 	virtual void OnExecution_Implementation(const TArray<UObject*>& InObjects, TArray<UObject*>& OutObjects) override;
 };
+
+UCLASS(Category = SelectionTransform, Meta = (DisplayName="Select Hierarchy", ToolTip = "Return immediate children or all the descendants of the selected objects") )
+class UDataprepHierarchySelectionTransform : public UDataprepSelectionTransform
+{
+	GENERATED_BODY()
+
+	UDataprepHierarchySelectionTransform()
+		: SelectionPolicy(EDataprepHierarchySelectionPolicy::ImmediateChildren)
+	{}
+
+protected:
+	virtual void OnExecution_Implementation(const TArray<UObject*>& InObjects, TArray<UObject*>& OutObjects) override;
+
+	UPROPERTY(EditAnywhere, BlueprintReadWrite, Category = "HierarchySelectionOptions", meta = (DisplayName = "Select", ToolTip = "Specify policy of hierarchical parsing of selected objects"))
+	EDataprepHierarchySelectionPolicy SelectionPolicy;
+};
\ No newline at end of file
diff --git a/Engine/Plugins/Enterprise/DatasmithImporter/Source/DatasmithImporter/DatasmithImporter.Build.cs b/Engine/Plugins/Enterprise/DatasmithImporter/Source/DatasmithImporter/DatasmithImporter.Build.cs
--- a/Engine/Plugins/Enterprise/DatasmithImporter/Source/DatasmithImporter/DatasmithImporter.Build.cs
+++ b/Engine/Plugins/Enterprise/DatasmithImporter/Source/DatasmithImporter/DatasmithImporter.Build.cs
@@ -7,10 +7,13 @@ namespace UnrealBuildTool.Rules
 		public DatasmithImporter(ReadOnlyTargetRules Target)
 			: base(Target)
 		{
+			OptimizeCode = CodeOptimization.InShippingBuildsOnly;
+
 			PrivateDependencyModuleNames.AddRange(
 				new string[]
 				{
 					"Analytics",
+					"ApplicationCore",
 					"CinematicCamera",
 					"ContentBrowser",
 					"Core",
@@ -35,6 +38,7 @@ namespace UnrealBuildTool.Rules
 					"MessageLog",
 					"MovieScene",
 					"MovieSceneTracks",
+					"PropertyEditor",
 					"RHI",
 					"RawMesh",
 					"Slate",
diff --git a/Engine/Plugins/Enterprise/DatasmithImporter/Source/DatasmithImporter/Private/DatasmithConsumer.cpp b/Engine/Plugins/Enterprise/DatasmithImporter/Source/DatasmithImporter/Private/DatasmithConsumer.cpp
--- a/Engine/Plugins/Enterprise/DatasmithImporter/Source/DatasmithImporter/Private/DatasmithConsumer.cpp
+++ b/Engine/Plugins/Enterprise/DatasmithImporter/Source/DatasmithImporter/Private/DatasmithConsumer.cpp
@@ -31,31 +31,40 @@
 
 #include "Algo/Count.h"
 #include "AssetRegistryModule.h"
+#include "AssetToolsModule.h"
 #include "CineCameraActor.h"
 #include "CineCameraComponent.h"
 #include "ComponentReregisterContext.h"
 #include "Components/HierarchicalInstancedStaticMeshComponent.h"
 #include "Components/PointLightComponent.h"
+#include "Editor.h"
 #include "EditorLevelUtils.h"
 #include "Engine/Level.h"
 #include "Engine/LevelStreamingAlwaysLoaded.h"
 #include "Engine/Light.h"
+#include "Engine/Selection.h"
 #include "Engine/StaticMesh.h"
 #include "Engine/StaticMeshActor.h"
 #include "Engine/World.h"
+#include "Factories/WorldFactory.h"
+#include "FileHelpers.h"
 #include "GameFramework/PhysicsVolume.h"
 #include "GameFramework/WorldSettings.h"
 #include "HAL/FileManager.h"
+#include "HAL/PlatformApplicationMisc.h"
 #include "IAssetTools.h"
 #include "Internationalization/Internationalization.h"
 #include "LevelSequence.h"
+#include "LevelUtils.h"
 #include "Materials/Material.h"
 #include "Materials/MaterialInstanceConstant.h"
 #include "Materials/MaterialInterface.h"
+#include "Misc/Guid.h"
 #include "Misc/MessageDialog.h"
 #include "Misc/PackageName.h"
 #include "Misc/Paths.h"
 #include "Misc/SecureHash.h"
+#include "ObjectTools.h"
 #include "UObject/Package.h"
 #include "UObject/SoftObjectPtr.h"
 #include "UObject/UnrealType.h"
@@ -80,7 +89,207 @@ namespace DatasmithConsumerUtils
 		UDatasmithContentBlueprintLibrary* DatasmithContentLibrary = Cast< UDatasmithContentBlueprintLibrary >( UDatasmithContentBlueprintLibrary::StaticClass()->GetDefaultObject() );
 		FString DatasmithUniqueId = DatasmithContentLibrary->GetDatasmithUserDataValueForKey( Object, UDatasmithAssetUserData::UniqueIdMetaDataKey );
 
-		return DatasmithUniqueId.IsEmpty() ? Object->GetName() : DatasmithUniqueId;
+		return DatasmithUniqueId.Len() == 0 ? Object->GetName() : DatasmithUniqueId;
+	}
+
+	void SaveMap(UWorld* WorldToSave);
+
+	const FString& GetMarker(UObject* Object, const FString& Name);
+
+	void SetMarker(UObject* Object, const FString& Name, const FString& Value);
+
+	void MoveActorsToLevel(const TArray<AActor*>& ActorsToMove, ULevel* DestLevel, TMap< FName, TSoftObjectPtr< AActor > >& ActorsMap, bool bDuplicate);
+
+	template<class AssetClass>
+	void SetMarker(const TMap<FName, TSoftObjectPtr< AssetClass >>& AssetMap, const FString& Name, const FString& Value)
+	{
+		for(const TPair<FName, TSoftObjectPtr< AssetClass >>& Entry : AssetMap)
+		{
+			if(UObject* Asset = Entry.Value.Get())
+			{
+				SetMarker(Asset, Name, Value);
+			}
+		}
+	}
+
+	template<class AssetClass>
+	void CollectAssetsToSave(TMap<FName, TSoftObjectPtr< AssetClass >>& AssetMap, TArray<UPackage*>& OutPackages)
+	{
+		if(AssetMap.Num() > 0)
+		{
+			OutPackages.Reserve(OutPackages.Num() + AssetMap.Num());
+
+			for(TPair<FName, TSoftObjectPtr< AssetClass >>& Entry : AssetMap)
+			{
+				if(UObject* Asset = Entry.Value.Get())
+				{
+					OutPackages.Add(Asset->GetOutermost());
+				}
+			}
+		}
+	}
+
+	template<class AssetClass>
+	void ApplyFolderDirective(TMap<FName, TSoftObjectPtr< AssetClass >>& AssetMap, const FString& RootPackagePath, TFunction<void(ELogVerbosity::Type, FText)> ReportCallback)
+	{
+		auto CanMoveAsset = [&ReportCallback](UObject* Source, UObject* Target) -> bool
+		{
+			// Overwrite existing owned asset with the new one
+			if( GetMarker(Source, UDatasmithConsumer::ConsumerMarkerID) == GetMarker(Target, UDatasmithConsumer::ConsumerMarkerID))
+			{
+				TArray<UObject*> ObjectsToReplace(&Target, 1);
+				ObjectTools::ForceReplaceReferences( Source, ObjectsToReplace );
+
+				Target->Rename(nullptr, GetTransientPackage(), REN_DontCreateRedirectors | REN_NonTransactional);
+
+				return true;
+			}
+
+			if(Source->GetClass() != Target->GetClass())
+			{
+				const FText AssetName = FText::FromString( Source->GetName() );
+				const FText AssetFolder = FText::FromString( FPaths::GetPath(Source->GetPathName()) );
+				const FText Message = FText::Format( LOCTEXT( "FolderDirective_ClassIssue", "Cannot move {0} to {1}. An asset with same name but different class exists"), AssetName, AssetFolder );
+				ReportCallback(ELogVerbosity::Error, Message);
+			}
+			else
+			{
+				const FText AssetName = FText::FromString( Source->GetName() );
+				const FText AssetFolder = FText::FromString( FPaths::GetPath(Source->GetPathName()) );
+				const FText Message = FText::Format( LOCTEXT( "FolderDirective_Overwrite", "Cannot move {0} to {1}. An asset with same name and same class exists"), AssetName, AssetFolder );
+				ReportCallback(ELogVerbosity::Error, Message);
+			}
+
+			return false;
+		};
+
+		TArray<UPackage*> PackagesToCheck;
+		TMap<FSoftObjectPath, FSoftObjectPath> AssetRedirectorMap;
+
+		for(TPair<FName, TSoftObjectPtr< AssetClass >>& Entry : AssetMap)
+		{
+			if(UObject* Asset = Entry.Value.Get())
+			{
+				const FString& OutputFolder = GetMarker(Asset, UDataprepContentConsumer::RelativeOutput);
+				if(OutputFolder.Len() > 0)
+				{
+					FString SourcePackagePath = Entry.Value->GetOuter()->GetPathName();
+					FString TargetPackagePath = FPaths::Combine(RootPackagePath, OutputFolder, Asset->GetName());
+					
+					FString PackageFilename;
+					FPackageName::TryConvertLongPackageNameToFilename( TargetPackagePath, PackageFilename, FPackageName::GetAssetPackageExtension() );
+
+					if( SourcePackagePath != TargetPackagePath)
+					{
+						bool bCanMove = true;
+
+						FString TargetAssetFullPath = TargetPackagePath + "." + Asset->GetName();
+						if(UObject* MemoryObject = FSoftObjectPath(TargetAssetFullPath).ResolveObject())
+						{
+							bCanMove = CanMoveAsset( Asset, MemoryObject);
+						}
+						else if(FPaths::FileExists(PackageFilename))
+						{
+							bCanMove = CanMoveAsset( Asset, FSoftObjectPath(TargetAssetFullPath).TryLoad());
+						}
+
+						if(bCanMove)
+						{
+							FSoftObjectPath& SoftObjectPathRef = AssetRedirectorMap.Emplace( Asset );
+
+							UPackage* Package = CreatePackage(nullptr, *TargetPackagePath);
+							Package->FullyLoad();
+
+							Asset->Rename(nullptr, Package, REN_DontCreateRedirectors | REN_NonTransactional);
+							Entry.Value = Asset;
+
+							SoftObjectPathRef = Asset;
+							PackagesToCheck.Add(Package);
+						}
+					}
+				}
+			}
+		}
+
+		if(AssetRedirectorMap.Num() > 0)
+		{
+			IAssetTools& AssetTools = FModuleManager::LoadModuleChecked<FAssetToolsModule>("AssetTools").Get();
+			AssetTools.RenameReferencingSoftObjectPaths(PackagesToCheck, AssetRedirectorMap);
+		}
+	}
+}
+
+const FString UDatasmithConsumer::ConsumerMarkerID = TEXT("DatasmithConsumer_UniqueID");
+
+UDatasmithConsumer::UDatasmithConsumer()
+	: DatasmithScene(nullptr)
+	, WorkingWorld(nullptr)
+	, PrimaryLevel(nullptr)
+{
+	if(!HasAnyFlags(RF_NeedLoad|RF_ClassDefaultObject))
+	{
+		UniqueID = FGuid::NewGuid().ToString(EGuidFormats::Short);
+	}
+}
+
+void UDatasmithConsumer::PostLoad()
+{
+	UDataprepContentConsumer::PostLoad();
+
+	// Update UniqueID for previous version of the consumer
+	if(HasAnyFlags(RF_WasLoaded))
+	{
+		bool bMarkDirty = false;
+		if(UniqueID.Len() == 0)
+		{
+			UniqueID = FGuid::NewGuid().ToString(EGuidFormats::Short);
+			bMarkDirty = true;
+		}
+
+		if(LevelName.Len() == 0)
+		{
+			LevelName = GetOuter()->GetName() + TEXT("_Map");
+		}
+
+		if(OutputLevelSoftPtr_DEPRECATED.ToString().Len() > 0)
+		{
+			OutputLevelSoftObject = OutputLevelSoftPtr_DEPRECATED.ToSoftObjectPath();
+
+			bMarkDirty = true;
+		}
+
+		if(OutputLevelSoftObject.GetAssetPathString().Len() == 0)
+		{
+			OutputLevelSoftObject = FSoftObjectPath(FPaths::Combine(TargetContentFolder, LevelName) + "." + LevelName);
+
+			bMarkDirty = true;
+		}
+
+		if(bMarkDirty)
+		{
+			const FText AssetName = FText::FromString( GetOuter()->GetName() );
+			const FText WarningMessage = FText::Format( LOCTEXT( "DataprepConsumerOldVersion", "{0} is from an old version and has been updated. Please save asset to complete update."), AssetName );
+			const FText NotificationText = FText::Format( LOCTEXT( "DataprepConsumerOldVersionNotif", "{0} is from an old version and has been updated."), AssetName );
+			LogWarning(WarningMessage);
+			//DataprepCorePrivateUtils::LogMessage( EMessageSeverity::Warning, WarningMessage, NotificationText );
+
+			GetOutermost()->SetDirtyFlag(true);
+		}
+	}
+}
+
+void UDatasmithConsumer::PostInitProperties()
+{
+	UDataprepContentConsumer::PostInitProperties();
+
+	if(!HasAnyFlags(RF_NeedLoad|RF_ClassDefaultObject))
+	{
+		if(LevelName.Len() == 0)
+		{
+			LevelName = GetOuter()->GetName() + TEXT("_Map");
+		}
+
+		OutputLevelSoftObject = FSoftObjectPath(FPaths::Combine( TargetContentFolder, LevelName) + TEXT(".") + LevelName);
 	}
 }
 
@@ -91,11 +300,12 @@ bool UDatasmithConsumer::Initialize()
 
 	ProgressTaskPtr->ReportNextStep( LOCTEXT( "DatasmithImportFactory_Initialize", "Preparing world ...") );
 
-	UpdateScene();
-
-	MoveLevel();
+	if(!CheckOutputDirectives())
+	{
+		return false;
+	}
 
-	UpdateLevel();
+	UpdateScene();
 
 	UPackage* ParentPackage = CreatePackage( nullptr, *GetTargetPackagePath() );
 	ParentPackage->FullyLoad();
@@ -180,30 +390,30 @@ bool UDatasmithConsumer::Initialize()
 		}
 	}
 
-	if ( !BuildContexts( Context.WorldPtr.Get() ) )
+	CreateWorld();
+
+	if ( !BuildContexts() )
 	{
 		return false;
 	}
 
 	// Check if the finalize should be threated as a reimport
-	if (FDatasmithImporterUtils::FindSceneActors( ImportContextPtr->ActorsContext.FinalWorld, ImportContextPtr->SceneAsset).Num() > 0 )
+	TArray< ADatasmithSceneActor* > SceneActors = FDatasmithImporterUtils::FindSceneActors( ImportContextPtr->ActorsContext.FinalWorld, ImportContextPtr->SceneAsset);
+	if (SceneActors.Num() > 0 )
 	{
 		ADatasmithSceneActor* FoundSceneActor = nullptr;
-		for( AActor* Actor : ImportContextPtr->ActorsContext.FinalWorld->GetCurrentLevel()->Actors )
+		for( ADatasmithSceneActor* SceneActor : SceneActors )
 		{
-			if( ADatasmithSceneActor* SceneActor = Cast<ADatasmithSceneActor>(Actor) )
+			if( SceneActor && SceneActor->Scene == DatasmithScene )
 			{
-				if( SceneActor->Scene == DatasmithScene )
-				{
-					FoundSceneActor = SceneActor;
-					break;
-				}
+				FoundSceneActor = SceneActor;
+				break;
 			}
 		}
 
 		if( FoundSceneActor == nullptr )
 		{
-			//Create a new datasmith scene actor in the targeted level
+			//Create a new Datasmith scene actor in the targeted level
 			FActorSpawnParameters SpawnParameters;
 			SpawnParameters.Template = ImportContextPtr->ActorsContext.ImportSceneActor;
 			ADatasmithSceneActor* DestinationSceneActor = Cast< ADatasmithSceneActor >(ImportContextPtr->ActorsContext.FinalWorld->SpawnActor< ADatasmithSceneActor >(SpawnParameters));
@@ -235,24 +445,260 @@ bool UDatasmithConsumer::Run()
 	ProgressTaskPtr->ReportNextStep( LOCTEXT( "DatasmithImportFactory_Finalize", "Finalizing commit ...") );
 	FDatasmithImporter::FinalizeImport( *ImportContextPtr, TSet<UObject*>() );
 
-	// Store the level name for subsequent call to Run
-	LastLevelName = LevelName;
+	// Apply UDataprepConsumerUserData directives for assets
+	UDatasmithScene* SceneAsset = ImportContextPtr->SceneAsset;
+
+	TFunction<void(ELogVerbosity::Type, FText)> ReportFunc = [&](ELogVerbosity::Type Verbosity, FText Message)
+	{
+		switch(Verbosity)
+		{
+			case ELogVerbosity::Warning:
+			LogWarning(Message);
+			break;
+
+			case ELogVerbosity::Error:
+			LogError(Message);
+			break;
+
+			default:
+			LogInfo(Message);
+			break;
+		}
+	};
+
+	DatasmithConsumerUtils::SetMarker(SceneAsset->Textures, UDatasmithConsumer::ConsumerMarkerID, UniqueID);
+	DatasmithConsumerUtils::ApplyFolderDirective(SceneAsset->Textures, TargetContentFolder, ReportFunc );
+
+	DatasmithConsumerUtils::SetMarker(SceneAsset->StaticMeshes, UDatasmithConsumer::ConsumerMarkerID, UniqueID);
+	DatasmithConsumerUtils::ApplyFolderDirective(SceneAsset->StaticMeshes, TargetContentFolder, ReportFunc );
+
+	DatasmithConsumerUtils::SetMarker(SceneAsset->Materials, UDatasmithConsumer::ConsumerMarkerID, UniqueID);
+	DatasmithConsumerUtils::ApplyFolderDirective(SceneAsset->Materials, TargetContentFolder, ReportFunc );
+
+	DatasmithConsumerUtils::SetMarker(SceneAsset->MaterialFunctions, UDatasmithConsumer::ConsumerMarkerID, UniqueID);
+	DatasmithConsumerUtils::ApplyFolderDirective(SceneAsset->MaterialFunctions, TargetContentFolder, ReportFunc );
+
+	DatasmithConsumerUtils::SetMarker(SceneAsset->LevelSequences, UDatasmithConsumer::ConsumerMarkerID, UniqueID);
+	DatasmithConsumerUtils::ApplyFolderDirective(SceneAsset->LevelSequences, TargetContentFolder, ReportFunc );
+
+	DatasmithConsumerUtils::SetMarker(SceneAsset->LevelVariantSets, UDatasmithConsumer::ConsumerMarkerID, UniqueID);
+	DatasmithConsumerUtils::ApplyFolderDirective(SceneAsset->LevelVariantSets, TargetContentFolder, ReportFunc );
+
+	// Apply UDataprepConsumerUserData directives for actors
+	ApplySubLevelDirective();
+
+	return FinalizeRun();
+}
+
+bool UDatasmithConsumer::FinalizeRun()
+{
+	UDatasmithScene* SceneAsset = ImportContextPtr->SceneAsset;
+
+	// Save all assets
+	TArray<UPackage*> PackagesToSave;
+	PackagesToSave.Add(DatasmithScene->GetOutermost());
+
+	DatasmithConsumerUtils::CollectAssetsToSave(SceneAsset->Textures, PackagesToSave);
+	DatasmithConsumerUtils::CollectAssetsToSave(SceneAsset->MaterialFunctions, PackagesToSave);
+	DatasmithConsumerUtils::CollectAssetsToSave(SceneAsset->Materials, PackagesToSave);
+	DatasmithConsumerUtils::CollectAssetsToSave(SceneAsset->StaticMeshes, PackagesToSave);
+	DatasmithConsumerUtils::CollectAssetsToSave(SceneAsset->LevelSequences, PackagesToSave);
+	DatasmithConsumerUtils::CollectAssetsToSave(SceneAsset->LevelVariantSets, PackagesToSave);
+
+	const bool bCheckDirty = false;
+	const bool bPromptToSave = false;
+	const FEditorFileUtils::EPromptReturnCode Return = FEditorFileUtils::PromptForCheckoutAndSave(PackagesToSave, bCheckDirty, bPromptToSave);
+
+	// Save secondary levels
+	const TArray<ULevel*>& Levels = WorkingWorld->GetLevels();
+
+	for(ULevel* Level : Levels)
+	{
+		if(Level)
+		{
+			UWorld* World = Cast<UWorld>(Level->GetOuter());
+
+			if(World != WorkingWorld.Get())
+			{
+				DatasmithConsumerUtils::SaveMap(World);
+			}
+		}
+	}
+
+	// Save primary level now
+	WorkingWorld->PersistentLevel = PrimaryLevel;
+	DatasmithConsumerUtils::SaveMap(WorkingWorld.Get());
+
+	return true;
+}
+
+
+bool UDatasmithConsumer::CreateWorld()
+{
+	ensure(!WorkingWorld.IsValid());
+
+	UWorld* World = Cast<UWorld>(OutputLevelSoftObject.TryLoad());
+	if(World)
+	{
+		World->SetFlags(RF_Public | RF_Transactional |  RF_Standalone);
+
+		FWorldContext& WorldContext = GEngine->CreateNewWorldContext(World->WorldType);
+		WorldContext.SetCurrentWorld(World);
+
+		// Load all the secondary levels of the world.
+		World->LoadSecondaryLevels(true);
+
+		// Check that all secondary levels have been added to world
+		for(ULevelStreaming* StreamingLevel : World->GetStreamingLevels())
+		{
+			if(ULevel* SecondaryLevel = StreamingLevel->GetLoadedLevel())
+			{
+				if(!World->ContainsLevel(SecondaryLevel))
+				{
+					World->AddLevel(SecondaryLevel);
+				}
+			}
+		}
+	}
+	else
+	{
+		UPackage* Package = CreatePackage(nullptr, *OutputLevelSoftObject.GetLongPackageName());
+		Package->FullyLoad();
+		Package->SetFlags(RF_Public);
+
+		World = NewObject< UWorld >(Package, *OutputLevelSoftObject.GetAssetName(), RF_Public | RF_Transactional |  RF_Standalone);
+		World->WorldType = EWorldType::Inactive;
+
+		FWorldContext& WorldContext = GEngine->CreateNewWorldContext(World->WorldType);
+		WorldContext.SetCurrentWorld(World);
+
+		World->InitializeNewWorld(UWorld::InitializationValues()
+			.AllowAudioPlayback(false)
+			.CreateAISystem(false)
+			.CreateNavigation(false)
+			.CreatePhysicsScene(false)
+			.RequiresHitProxies(false)
+			.ShouldSimulatePhysics(false)
+			.SetTransactional(false));
+	}
+
+	if(World == nullptr)
+	{
+		ensure(false);
+		return false;
+	}
+
+	TArray<ULevel*> Levels = World->GetLevels();
+
+	// Find level associated with this consumer
+	PrimaryLevel = nullptr;
+	for(ULevel* Level : Levels)
+	{
+		if(Level && Level->GetOuter()->GetName() == LevelName)
+		{
+			PrimaryLevel = Level;
+			break;
+		}
+	}
+	ensure(PrimaryLevel);
+
+	DatasmithConsumerUtils::SetMarker(PrimaryLevel, ConsumerMarkerID, UniqueID);
+	PrimaryLevel->bIsVisible = true;
+
+	// If there is more than one level, move all actors of the world to the main level.
+	// The call to FinalizeRun will take care of redistributing to the sub-levels if applicable
+	if(Levels.Num() > 1)
+	{
+		// Get the ADatasmithSceneActor of the world if it exists
+		TArray< ADatasmithSceneActor* > SceneActors = FDatasmithImporterUtils::FindSceneActors( World, DatasmithScene.Get());
+		TMap< FName, TSoftObjectPtr< AActor > > EmptyRelatedActors;
+		TMap< FName, TSoftObjectPtr< AActor > >& RelatedActors = SceneActors.Num() > 0 ? SceneActors[0]->RelatedActors : EmptyRelatedActors;
+
+		for(ULevel* Level : Levels)
+		{
+			if(Level != PrimaryLevel)
+			{
+				Level->bIsVisible = true;
+
+				// Collect actors to copy to primary level
+				TArray<AActor*> ActorsToCopy;
+				for(AActor* Actor : Level->Actors)
+				{
+					if(Actor && Actor->GetRootComponent() && !Actor->IsA<AWorldSettings>() && !Actor->IsA<APhysicsVolume>() && !Actor->IsA<ABrush>() )
+					{
+						ActorsToCopy.Add(Actor);
+					}
+				}
+
+				DatasmithConsumerUtils::MoveActorsToLevel( ActorsToCopy, PrimaryLevel, RelatedActors, true);
+
+				World->RemoveLevel(Level);
+
+				// Make sure the world and package are properly discarded
+				UPackage* LevelPackage = Level->GetOutermost();
+				UWorld* LevelWorld = Cast<UWorld>(Level->GetOuter());
+
+				// Move the world to the transient package
+				LevelWorld->Rename(nullptr, GetTransientPackage(), REN_NonTransactional | REN_DontCreateRedirectors);
+
+				// Empty the world from all its content
+				LevelWorld->DestroyWorld( true );
+				GEngine->DestroyWorldContext( LevelWorld );
+
+				// Indicates world is good for garbage collect
+				LevelWorld->ClearFlags(RF_Standalone | RF_Public | RF_Transactional);
+				LevelWorld->SetFlags(RF_Transient);
+				LevelWorld->MarkPendingKill();
+
+				// Indicates level's package is good for garbage collect
+				LevelPackage->SetDirtyFlag(false);
+				LevelPackage->ClearFlags(RF_Standalone | RF_Public | RF_Transactional);
+				LevelPackage->SetFlags(RF_Transient);
+				LevelPackage->MarkPendingKill();
+			}
+		}
+
+		World->ClearStreamingLevels();
+
+		// Collect garbage to clear out the discarded world(s) and level(s) 
+		CollectGarbage( GARBAGE_COLLECTION_KEEPFLAGS );
+	}
+
+	World->PersistentLevel = PrimaryLevel;
+	World->SetCurrentLevel(PrimaryLevel);
+
+	WorkingWorld = TStrongObjectPtr<UWorld>(World);
 
 	return true;
 }
 
+void UDatasmithConsumer::ClearWorld()
+{
+	if(WorkingWorld.IsValid())
+	{
+		UWorld* WorldToDelete = WorkingWorld.Get();
+		WorkingWorld.Reset();
+
+		WorldToDelete->PersistentLevel = PrimaryLevel;
+		WorldToDelete->SetCurrentLevel(PrimaryLevel);
+
+		PrimaryLevel = nullptr;
+
+		WorldToDelete->DestroyWorld( true );
+		GEngine->DestroyWorldContext( WorldToDelete );
+
+		// Collect garbage to clear out the destroyed level
+		CollectGarbage( GARBAGE_COLLECTION_KEEPFLAGS );
+	}
+}
+
 void UDatasmithConsumer::Reset()
 {
 	ImportContextPtr.Reset();
 	ProgressTaskPtr.Reset();
 	UDataprepContentConsumer::Reset();
 
-	// Restore previous current level
-	if( PreviousCurrentLevel != nullptr )
-	{
-		GWorld->SetCurrentLevel( PreviousCurrentLevel );
-		PreviousCurrentLevel = nullptr;
-	}
+	ClearWorld();
 }
 
 const FText& UDatasmithConsumer::GetLabel() const
@@ -265,7 +711,7 @@ const FText& UDatasmithConsumer::GetDescription() const
 	return DatasmithConsumerDescription;
 }
 
-bool UDatasmithConsumer::BuildContexts( UWorld* ImportWorld )
+bool UDatasmithConsumer::BuildContexts()
 {
 	const FString FilePath = FPaths::Combine( FPaths::ProjectIntermediateDir(), ( DatasmithScene->GetName() + TEXT( ".udatasmith" ) ) );
 
@@ -274,7 +720,7 @@ bool UDatasmithConsumer::BuildContexts( UWorld* ImportWorld )
 	// Update import context with consumer's data
 	ImportContextPtr->Options->BaseOptions.SceneHandling = EDatasmithImportScene::CurrentLevel;
 	ImportContextPtr->SceneAsset = DatasmithScene.Get();
-	ImportContextPtr->ActorsContext.ImportWorld = ImportWorld;
+	ImportContextPtr->ActorsContext.ImportWorld = Context.WorldPtr.Get();
 	ImportContextPtr->Scene = FDatasmithSceneFactory::CreateScene( *DatasmithScene->GetName() );
 	ImportContextPtr->SceneName = ImportContextPtr->Scene->GetName();
 
@@ -330,11 +776,11 @@ bool UDatasmithConsumer::BuildContexts( UWorld* ImportWorld )
 	ImportContextPtr->AssetsContext.LevelVariantSetsImportPackage.Reset();
 
 	// Set the destination world as the one in the level editor
-	ImportContextPtr->ActorsContext.FinalWorld = GWorld;
+	ImportContextPtr->ActorsContext.FinalWorld = WorkingWorld.Get();
 
 	// Initialize ActorsContext's UniqueNameProvider with actors in the GWorld not the Import world
 	ImportContextPtr->ActorsContext.UniqueNameProvider = FDatasmithActorUniqueLabelProvider();
-	ImportContextPtr->ActorsContext.UniqueNameProvider.PopulateLabelFrom( GWorld );
+	ImportContextPtr->ActorsContext.UniqueNameProvider.PopulateLabelFrom( ImportContextPtr->ActorsContext.FinalWorld );
 
 	// Add assets as if they have been imported using the current import context
 	DatasmithConsumerUtils::AddAssetsToContext( *ImportContextPtr, Context.Assets );
@@ -344,104 +790,107 @@ bool UDatasmithConsumer::BuildContexts( UWorld* ImportWorld )
 	return true;
 }
 
-ULevel * UDatasmithConsumer::FindLevel( const FString& InLevelName )
+bool UDatasmithConsumer::SetLevelNameImplementation(const FString& InLevelName, FText& OutReason, const bool bIsAutomated)
 {
-	UWorld* FinalWorld = GWorld;
+	if(InLevelName.Len() == 0)
+	{
+		OutReason = LOCTEXT( "DatasmithConsumer_NameEmpty", "The level name is empty. Please enter a valid name." );
+		return false;
+	}
 
-	FSoftObjectPath LevelObjectPath( FPaths::Combine( TargetContentFolder, InLevelName ) );
-	UObject* Object = LevelObjectPath.ResolveObject();
-	ULevel* Level = Cast<ULevel>( Object );
+	if(!CanCreateLevel(TargetContentFolder, InLevelName, !bIsAutomated && !IsRunningCommandlet()))
+	{
+		return false;
+	}
 
-	for( const ULevelStreaming* LevelStreaming : FinalWorld->GetStreamingLevels() )
+	if(SetOutputLevel(InLevelName))
 	{
-		if( LevelStreaming->GetWorldAssetPackageName() == LevelObjectPath.ToString() )
-		{
-			return LevelStreaming->GetLoadedLevel();
-		}
+		Modify();
+
+		LevelName = InLevelName;
+
+		OnChanged.Broadcast();
+
+		return true;
 	}
 
-	return Level;
+	// Warn user new name has not been set
+	OutReason = FText::Format( LOCTEXT("DatasmithConsumer_BadLevelName", "Cannot create level named {0}."), FText::FromString( InLevelName ) );
+
+	return false;
 }
 
-bool UDatasmithConsumer::SetLevelNameImplementation(const FString & InLevelName, FText& OutReason, const bool bIsAutomated)
+bool UDatasmithConsumer::CanCreateLevel(const FString& RequestedFolder, const FString& RequestedName, const bool bShowDialog)
 {
-	FString NewLevelName = InLevelName;
+	FSoftObjectPath ObjectPath(FPaths::Combine(RequestedFolder, RequestedName) + TEXT(".") + RequestedName);
 
-	bool bValidLevelName = false;
-	OutReason = FText();
+	FString PackageFilename;
+	FPackageName::TryConvertLongPackageNameToFilename( ObjectPath.GetLongPackageName(), PackageFilename, FPackageName::GetMapPackageExtension() );
 
-	// Check if a new level can be used with the new name and current limitations
-	if( !NewLevelName.IsEmpty() && NewLevelName.Compare( TEXT("current"), ESearchCase::IgnoreCase ) != 0 )
+	if(FPaths::FileExists(PackageFilename))
 	{
-		// Sub-level of sub-level is not supported yet
-		// #ueent_todo: sub-level of sub-level
-		if( InLevelName.Contains( TEXT("/") ) || InLevelName.Contains( TEXT("\\") ))
-		{
-			OutReason = LOCTEXT( "DatasmithConsumer_SubLevel", "Sub-level of sub-levels is not supported yet" );
-		}
-		// Try to see if there is any issue to eventually create this level, i.e. name collision
-		else if( FindLevel( InLevelName ) == nullptr )
+		if(UWorld* World = Cast<UWorld>(ObjectPath.TryLoad()))
 		{
-			FSoftObjectPath LevelObjectPath( FPaths::Combine( TargetContentFolder, InLevelName ) );
-
-			if( StaticFindObject( nullptr, ANY_PACKAGE, *LevelObjectPath.ToString(), true) )
+			if(DatasmithConsumerUtils::GetMarker(World->PersistentLevel, ConsumerMarkerID) != UniqueID)
 			{
-				OutReason = LOCTEXT( "DatasmithConsumer_LevelExists", "A object with that name already exists. Please choose another name." );
-			}
+				if(bShowDialog)
+				{
+					const FTextFormat Format(LOCTEXT("DatasmithConsumer_SetTargetContentFolder_Overwrite_Dlg", "Level {0} already exists in {1}.\n\nDo you want to overwrite it?"));
+					const FText WarningMessage = FText::Format( Format, FText::FromString(RequestedName), FText::FromString(RequestedFolder));
+					const FText DialogTitle( LOCTEXT("DatasmithConsumer_Overwrite_DlgTitle", "Warning - Level already exists") );
 
-			// #ueent_todo: Check if persistent level is locked, etc
+					if(FMessageDialog::Open(EAppMsgType::YesNo, WarningMessage, &DialogTitle) != EAppReturnType::Yes)
+					{
+						return false;
+					}
+				}
+				else
+				{
+					const FTextFormat Format(LOCTEXT("DatasmithConsumer_SetTargetContentFolder_Overwrite", "Level {0} already exists in {1}.It will be overwritten."));
+					const FText WarningMessage = FText::Format( Format, FText::FromString(RequestedName), FText::FromString(RequestedFolder));
+					LogWarning(WarningMessage);
+				}
+			}
 		}
-
-		// Good to go if no error documented
-		bValidLevelName = OutReason.IsEmpty();
 	}
-	// New name of level is empty or keyword 'current' used
-	else if( !LevelName.IsEmpty() )
+	else if(FDatasmithImporterUtils::CanCreateAsset(ObjectPath.GetAssetPathString(), UWorld::StaticClass()) != FDatasmithImporterUtils::EAssetCreationStatus::CS_CanCreate)
 	{
-		NewLevelName = TEXT("");
-		bValidLevelName = true;
-	}
+		const FTextFormat Format(LOCTEXT("DatasmithConsumer_SetTargetContentFolder_CantCreate_Dlg", "Cannot create level {0} in folder {1}."));
+		const FText Message = FText::Format( Format, FText::FromString(RequestedName), FText::FromString(RequestedFolder));
 
-	if(bValidLevelName)
-	{
-		Modify();
+		if(bShowDialog)
+		{
+			const FText DialogTitle( LOCTEXT("DatasmithConsumer_CantCreate_DlgTitle", "Warning - Cannot create level") );
 
-		LevelName = NewLevelName;
+			FMessageDialog::Open(EAppMsgType::Ok, Message, &DialogTitle);
 
-		OnChanged.Broadcast();
+			return false;
+		}
+		else
+		{
+			LogError(Message);
+		}
 	}
 
-	return bValidLevelName;
+	return true;
 }
 
 bool UDatasmithConsumer::SetTargetContentFolderImplementation(const FString& InTargetContentFolder, FText& OutFailureReason, const bool bIsAutomated)
 {
+	if(!CanCreateLevel(InTargetContentFolder, LevelName, !bIsAutomated && !IsRunningCommandlet()))
+	{
+		return false;
+	}
 
 	if ( Super::SetTargetContentFolderImplementation( InTargetContentFolder, OutFailureReason, bIsAutomated ) )
 	{
-		// Warn user if related Datasmith scene is not in package path and force re-creation of Datasmith scene
-		if ( DatasmithScene )
-		{ 
-			const FString DatasmithScenePath = FPaths::GetPath( DatasmithScene->GetPathName() );
-			if ( FPaths::GetPath( DatasmithScene->GetPathName() ) != TargetContentFolder)
-			{
-				DatasmithScene.Reset();
-
-				FText WarningMessage = FText::Format( LOCTEXT("DatasmithConsumer_NoSceneAsset", "Package path {0} different from path previously used, {1}.\nPrevious content will not be updated."), FText::FromString( TargetContentFolder ), FText::FromString( DatasmithScenePath ) );
-
-				if ( !bIsAutomated )
-				{
-					const FText DialogTitle( LOCTEXT("DatasmithConsumerDlgTitle", "Warning") );
+		// Inform user if related Datasmith scene is not in package path and force re-creation of Datasmith scene
+		const FText Message = FText::Format( LOCTEXT("DatasmithConsumer_SetTargetContentFolder", "Package path {0} different from path previously used. Previous content will not be updated."), FText::FromString( TargetContentFolder ) );
+		LogInfo(Message);
 
-					// We should move that on the text box widget directly
-					FMessageDialog::Open( EAppMsgType::Ok, WarningMessage, &DialogTitle );
-				}
-
-				UE_LOG(LogDatasmithImport, Log, TEXT("%s"), *WarningMessage.ToString());
-			}
-		}
+		DatasmithScene.Reset();
 
-		return true;
+		return SetOutputLevel( LevelName );
 	}
 
 	return false;
@@ -462,79 +911,354 @@ void UDatasmithConsumer::UpdateScene()
 		// Force re-creation of Datasmith scene
 		DatasmithScene.Reset();
 	}
+}
+
+bool UDatasmithConsumer::SetOutputLevel(const FString& InLevelName)
+{
+	if(InLevelName.Len() > 0)
+	{
+		Modify();
+
+		OutputLevelSoftObject = FSoftObjectPath(FPaths::Combine( TargetContentFolder, InLevelName) + TEXT(".") + InLevelName);
+
+		MarkPackageDirty();
+
+		OnChanged.Broadcast();
+
+		return true;
+	}
 
+	return false;
 }
 
-void UDatasmithConsumer::MoveLevel()
+ULevel* UDatasmithConsumer::FindOrAddLevel(const FString& InLevelName)
 {
-	// Do nothing if this is the First call to Run, DatasmithScene is null and LastLevelName is empty
-	// or the re-Run is using the same level
-	if( ( !DatasmithScene.IsValid() && LastLevelName.IsEmpty() ) || LastLevelName == LevelName )
+	FString LevelPackageName = FPaths::Combine(TargetContentFolder, InLevelName);
+
+	if(ULevelStreaming* StreamingLevel = FLevelUtils::FindStreamingLevel(WorkingWorld.Get(), *LevelPackageName))
 	{
-		return;
+		return StreamingLevel->GetLoadedLevel();
 	}
 
-	ULevel* Level = FindLevel( LevelName );
-	if( Level == nullptr )
+	ULevel* CurrentLevel = WorkingWorld->PersistentLevel;
+
+	// This level has not been added yet
+	FString PackageFilename;
+	FPackageName::TryConvertLongPackageNameToFilename( LevelPackageName, PackageFilename, FPackageName::GetMapPackageExtension() );
+
+	bool bCleanLevel = false;
+
+	ULevelStreaming* StreamingLevel = nullptr;
+	if(FPaths::FileExists(PackageFilename))
 	{
-		return;
+		FTransform LevelTransform;
+		StreamingLevel = UEditorLevelUtils::AddLevelToWorld(WorkingWorld.Get(), *LevelPackageName, ULevelStreamingAlwaysLoaded::StaticClass(), LevelTransform);
+		ensure(StreamingLevel);
+
+		WorkingWorld->LoadSecondaryLevels();
+		ensure(StreamingLevel->GetLoadedLevel());
+
+		bCleanLevel = true;
+	}
+	else
+	{
+		StreamingLevel = EditorLevelUtils::CreateNewStreamingLevelForWorld( *WorkingWorld, ULevelStreamingAlwaysLoaded::StaticClass(), *PackageFilename );
+		ensure(StreamingLevel);
 	}
 
-	// New level exists, search for DatasmithSceneActor associated with this consumer
-	ADatasmithSceneActor* FoundSceneActor = nullptr;
-	for( AActor* Actor : Level->Actors )
+	WorkingWorld->PersistentLevel = CurrentLevel;
+	WorkingWorld->SetCurrentLevel(CurrentLevel);
+
+	// Mark level as generated by this consumer
+	if(StreamingLevel)
 	{
-		if( ADatasmithSceneActor* SceneActor = Cast<ADatasmithSceneActor>(Actor) )
+		ULevel* NewLevel = StreamingLevel->GetLoadedLevel();
+
+		DatasmithConsumerUtils::SetMarker(NewLevel, ConsumerMarkerID, UniqueID);
+
+		if(bCleanLevel)
 		{
-			if( SceneActor->Scene == DatasmithScene.Get() )
+			// Clean up the level if it contains actors from previous execution
+			UWorld* LevelWorld = Cast<UWorld>(NewLevel->GetOuter());
+
+			TArray<AActor*> LevelActors = NewLevel->Actors;
+			for(AActor* Actor : LevelActors)
 			{
-				FoundSceneActor = SceneActor;
-				break;
+
+				if(Actor && Actor->GetRootComponent() && !Actor->IsA<AWorldSettings>() && !Actor->IsA<APhysicsVolume>() && !Actor->IsA<ABrush>() )
+				{
+					LevelWorld->DestroyActor(Actor, true);
+					Actor->UObject::Rename(nullptr, GetTransientPackage(), REN_DontCreateRedirectors | REN_ForceNoResetLoaders);
+				}
 			}
+
+			LevelWorld->CleanupWorld(false, false);
+			LevelWorld->CleanupActors();
+
+			DatasmithConsumerUtils::SaveMap(LevelWorld);
 		}
+
+		WorkingWorld->AddLevel(NewLevel);
+
+		return NewLevel;
 	}
+
+	return nullptr;
 }
 
-void UDatasmithConsumer::UpdateLevel()
+bool UDatasmithConsumer::CheckOutputDirectives()
 {
-	PreviousCurrentLevel = nullptr;
+	auto CanCreateAsset = [](const FString& AssetPathName,const UClass* AssetClass)
+	{
+		return FDatasmithImporterUtils::CanCreateAsset(AssetPathName, AssetClass) == FDatasmithImporterUtils::EAssetCreationStatus::CS_CanCreate;
+	};
 
-	if( !LevelName.IsEmpty() )
+	const bool bShowDialog = !Context.bSilentMode && !IsRunningCommandlet();
+
+	// Collect garbage to clear out the destroyed level
+	CollectGarbage( GARBAGE_COLLECTION_KEEPFLAGS );
+
+	// Collect all sub-levels to be created
+	TSet<FString> LevelsToCreate;
+	LevelsToCreate.Add(LevelName);
+
+	for(AActor* Actor : Context.WorldPtr->GetCurrentLevel()->Actors)
 	{
-		UWorld* FinalWorld = GWorld;
+		if(Actor)
+		{
+			const FString& OutputLevelName = DatasmithConsumerUtils::GetMarker(Actor->GetRootComponent(), UDataprepContentConsumer::RelativeOutput);
+			if(OutputLevelName.Len() > 0)
+			{
+				LevelsToCreate.Add(OutputLevelName);
+			}
+		}
+	}
 
-		ULevel* Level = FindLevel( LevelName );
+	bool bCannotCreateAsset = false;
+
+	// check if any of the levels to create or update is not opened in the level editor
+	TArray<FString> OpenedLevels;
+	UWorld* GlobalWorld = GWorld;
+	{
+		TArray<ULevel*> WorldLevels = GlobalWorld->GetLevels();
+		for(ULevel* Level : WorldLevels)
+		{
+			if(Level)
+			{
+				const FString WorldLevelName = Level->GetOuter()->GetName();
+				if(LevelsToCreate.Contains(WorldLevelName))
+				{
+					OpenedLevels.Add(WorldLevelName);
+				}
+			}
+		}
+	}
+
+	for(const FString& LevelToCreate : LevelsToCreate)
+	{
+		FSoftObjectPath AssetSoftObjectPath(FPaths::Combine(TargetContentFolder, LevelToCreate) + "." + LevelToCreate);
 
-		if( Level == nullptr )
+		if(!CanCreateAsset(AssetSoftObjectPath.GetAssetPathString(), UWorld::StaticClass()))
 		{
-			FSoftObjectPath LevelObjectPath( FPaths::Combine( TargetContentFolder, LevelName ) );
+			const FTextFormat TextFormat(LOCTEXT( "DatasmithConsumer_CannotCreateAsset", "Cannot create asset {0}. Commit will be aborted" ));
+			const FText Message = FText::Format( TextFormat, FText::FromString(AssetSoftObjectPath.GetAssetPathString()) );
+			LogError(Message);
 
+			bCannotCreateAsset = true;
+		}
+		// Check if umap file does not already exist. If so, user will be asked if he/she wants to overwrite it
+		else
+		{
 			FString PackageFilename;
-			FPackageName::TryConvertLongPackageNameToFilename( LevelObjectPath.ToString(), PackageFilename, FPackageName::GetMapPackageExtension() );
-			if( ULevelStreaming* LevelStreaming = EditorLevelUtils::CreateNewStreamingLevelForWorld( *FinalWorld, ULevelStreamingAlwaysLoaded::StaticClass(), *PackageFilename ) )
+			FPackageName::TryConvertLongPackageNameToFilename( AssetSoftObjectPath.GetLongPackageName(), PackageFilename, FPackageName::GetMapPackageExtension() );
+			if(FPaths::FileExists(PackageFilename))
 			{
-				Level = LevelStreaming->GetLoadedLevel();
+				const FTextFormat TextFormat(LOCTEXT( "DatasmithConsumer_UMapAlreadyExists", "Level {0} already exists" ));
+				const FText Message = FText::Format( TextFormat, FText::FromString(LevelToCreate) );
+				LogInfo(Message);
 			}
-			else
+		}
+	}
+
+	// Abort commit if any level is opened
+	if(OpenedLevels.Num() > 0)
+	{
+		FText OpenedLevelsText;
+		if(OpenedLevels.Num() == 1)
+		{
+			const FTextFormat TextFormat(LOCTEXT( "DatasmithConsumer_OneLevelOpened", "level {0} is opened" ));
+			OpenedLevelsText = FText::Format( TextFormat, FText::FromString(OpenedLevels[0]) );
+		}
+		else
+		{
+			const FString Separator(TEXT(", "));
+			FString OpenedLevelListString = OpenedLevels[0];
+			for(int32 Index = 1; Index < OpenedLevels.Num(); ++Index)
 			{
-				FText Message = LOCTEXT( "DatasmithConsumer_UpdateLevel", "Cannot create level..." );
-				LogWarning( Message );
-				Level = FinalWorld->PersistentLevel;
+				OpenedLevelListString += Separator + OpenedLevels[Index];
 			}
 
-			check( Level );
+			const FTextFormat TextFormat(LOCTEXT( "DatasmithConsumer_MultipleLevelOpened", "levels {0} are opened"));
+			OpenedLevelsText = FText::Format( TextFormat, FText::FromString(OpenedLevelListString) );
+		}
+
+		const FText Message = FText::Format( LOCTEXT( "DatasmithConsumer_OpenAbortCommit", "Cannot proceed with commit because {0}.\nPlease close any editor using this level and commit again" ), OpenedLevelsText );
+		if(bShowDialog)
+		{
+			const FText Title( LOCTEXT( "DatasmithConsumer_OpenAbortCommitTitle", "Main level is opened" ) );
+			FMessageDialog::Open( EAppMsgType::Ok, Message, &Title );
+		}
+		else
+		{
+			LogError(Message);
+		}
+
+		return false;
+	}
+	
+	for(const TWeakObjectPtr< UObject >& AssetPtr : Context.Assets)
+	{
+		if(UObject* Asset = AssetPtr.Get())
+		{
+			const FString& OutputFolder = DatasmithConsumerUtils::GetMarker(Asset, UDataprepContentConsumer::RelativeOutput);
+			if(OutputFolder.Len() > 0)
+			{
+				FString AssetName = Asset->GetName();
+				FSoftObjectPath AssetSoftObjectPath(FPaths::Combine(TargetContentFolder, OutputFolder, AssetName) + "." + AssetName);
+
+				if(Asset->GetPathName() != AssetSoftObjectPath.GetLongPackageName())
+				{
+					if(!CanCreateAsset(AssetSoftObjectPath.GetAssetPathString(), Asset->GetClass()))
+					{
+						const FTextFormat TextFormat(LOCTEXT( "DatasmithConsumer_CannotCreateAsset", "Cannot create asset {0}. Commit will be aborted" ));
+						const FText Message = FText::Format( TextFormat, FText::FromString(AssetSoftObjectPath.GetAssetPathString()) );
+						LogError(Message);
+
+						bCannotCreateAsset = true;
+					}
+				}
+			}
 		}
+	}
+
+	if(bCannotCreateAsset)
+	{
 
-		if( Level != FinalWorld->GetCurrentLevel() )
+		const FText Message = LOCTEXT( "DatasmithConsumer_CreateAbortCommit", "Cannot proceed with commit because some assets and/or levels cannot be created.\nCheck your log for details, fix all issues and commit again" );
+		
+		if(bShowDialog)
 		{
-			PreviousCurrentLevel = FinalWorld->GetCurrentLevel();
-			FinalWorld->SetCurrentLevel( Level );
+			const FText Title( LOCTEXT( "DatasmithConsumer_CreateAbortCommitTitle", "Cannot create some assets" ) );
+			FMessageDialog::Open( EAppMsgType::Ok, Message, &Title );
+		}
+		else
+		{
+			LogError(Message);
+		}
+
+		return false;
+	}
+
+	return true;
+}
+
+void UDatasmithConsumer::ApplySubLevelDirective()
+{
+	TMap< FName, TSoftObjectPtr< AActor > >& RelatedActors = ImportContextPtr->ActorsContext.CurrentTargetedScene->RelatedActors;
+
+	TMap<FString, ULevel*> LevelMap;
+	TMap<ULevel*, TArray<AActor*>> ActorsToMove;
+
+	LevelMap.Add(LevelName, PrimaryLevel);
+	ActorsToMove.Add(PrimaryLevel);
+
+	for(TPair< FName, TSoftObjectPtr< AActor > >& Entry : RelatedActors)
+	{
+		if(AActor* Actor = Entry.Value.Get())
+		{
+			ULevel* TargetLevel = PrimaryLevel;
+
+			const FString& OutputDirectiveName = DatasmithConsumerUtils::GetMarker(Actor->GetRootComponent(), UDataprepContentConsumer::RelativeOutput);
+			if(OutputDirectiveName.Len() > 0 && OutputDirectiveName != LevelName)
+			{
+				ULevel* Level = nullptr;
+				if(ULevel** OutputLevelPtr = LevelMap.Find(OutputDirectiveName))
+				{
+					Level = *OutputLevelPtr;
+				}
+				else
+				{
+					Level = FindOrAddLevel(OutputDirectiveName);
+
+					if(Level)
+					{
+						// Tag new level as owned by consumer
+						LevelMap.Add(OutputDirectiveName, Level);
+						DatasmithConsumerUtils::SetMarker(Level, ConsumerMarkerID, UniqueID);
+					}
+					else
+					{
+						FText Message = LOCTEXT( "DatasmithConsumer_ApplySubLevelDirective", "Cannot create level..." );
+						LogWarning( Message );
+					}
+				}
+
+				if(Level)
+				{
+					TargetLevel = Level;
+				}
+
+			}
+
+			if(Actor->GetLevel() != TargetLevel)
+			{
+				ActorsToMove.FindOrAdd(TargetLevel).Add(Actor);
+			}
 		}
 	}
+
+	for(TPair<ULevel*, TArray<AActor*>> Entry : ActorsToMove)
+	{
+		DatasmithConsumerUtils::MoveActorsToLevel( Entry.Value, Entry.Key, RelatedActors, false);
+	}
 }
 
 namespace DatasmithConsumerUtils
 {
+	const FString& GetMarker(UObject* Object, const FString& Name)
+	{
+		if ( IInterface_AssetUserData* AssetUserDataInterface = Cast< IInterface_AssetUserData >( Object ) )
+		{
+			UDataprepConsumerUserData* DataprepContentUserData = AssetUserDataInterface->GetAssetUserData< UDataprepConsumerUserData >();
+
+			if ( DataprepContentUserData )
+			{
+				return DataprepContentUserData->GetMarker(Name);
+			}
+		}
+
+		static FString NullString;
+
+		return NullString;
+	}
+
+	void SetMarker(UObject* Object, const FString& Name, const FString& Value)
+	{
+		if ( IInterface_AssetUserData* AssetUserDataInterface = Cast< IInterface_AssetUserData >( Object ) )
+		{
+			UDataprepConsumerUserData* DataprepContentUserData = AssetUserDataInterface->GetAssetUserData< UDataprepConsumerUserData >();
+
+			if ( !DataprepContentUserData )
+			{
+				EObjectFlags Flags = RF_Public;
+				DataprepContentUserData = NewObject< UDataprepConsumerUserData >( Object, NAME_None, Flags );
+				AssetUserDataInterface->AddAssetUserData( DataprepContentUserData );
+			}
+
+			return DataprepContentUserData->AddMarker(Name, Value);
+		}
+	}
+
 	void ConvertSceneActorsToActors( FDatasmithImportContext& ImportContext )
 	{
 		UWorld* ImportWorld = ImportContext.ActorsContext.ImportWorld;
@@ -608,6 +1332,26 @@ namespace DatasmithConsumerUtils
 			// Attach new actor to root scene actor
 			ActorRootComponent->AttachToComponent( NewSceneActorRootComponent, FAttachmentTransformRules::KeepRelativeTransform );
 
+			// Copy AssetUserData - it is done by known classes but should be improved
+			if ( IInterface_AssetUserData* SourceAssetUserDataInterface = Cast< IInterface_AssetUserData >( SceneActorRootComponent ) )
+			{
+				if(IInterface_AssetUserData* TargetAssetUserDataInterface = Cast< IInterface_AssetUserData >(ActorRootComponent))
+				{
+					if(UAssetUserData* SourceDatasmithUserData = SourceAssetUserDataInterface->GetAssetUserDataOfClass(UDatasmithAssetUserData::StaticClass()))
+					{
+						UAssetUserData* TargetDatasmithUserData = DuplicateObject<UAssetUserData>(SourceDatasmithUserData, ActorRootComponent);
+						TargetAssetUserDataInterface->AddAssetUserData(TargetDatasmithUserData);
+					}
+
+					if(UAssetUserData* SourceConsumerUserData = SourceAssetUserDataInterface->GetAssetUserDataOfClass(UDataprepConsumerUserData::StaticClass()))
+					{
+						UAssetUserData* TargetConsumerUserData = DuplicateObject<UAssetUserData>(SourceConsumerUserData, ActorRootComponent);
+						TargetAssetUserDataInterface->AddAssetUserData(TargetConsumerUserData);
+					}
+				}
+			}
+
+
 			// Delete root scene actor since it is not needed anymore
 			ImportWorld->DestroyActor( SceneActor, false, true );
 			SceneActor->UnregisterAllComponents();
@@ -770,7 +1514,7 @@ namespace DatasmithConsumerUtils
 					// Clean up static meshes which have incomplete render data.
 					if(StaticMesh->RenderData.IsValid() && !StaticMesh->RenderData->IsInitialized())
 					{
-						StaticMesh->RenderData.Reset(nullptr);
+						StaticMesh->RenderData.Reset();
 					}
 
 					if(FDatasmithObjectTemplateUtils::GetObjectTemplate<UDatasmithStaticMeshTemplate>( StaticMesh ) == nullptr)
@@ -848,6 +1592,205 @@ namespace DatasmithConsumerUtils
 			}
 		}
 	}
+
+	void SaveMap(UWorld* WorldToSave)
+	{
+		const bool bHasStandaloneFlag = WorldToSave->HasAnyFlags(RF_Standalone);
+		FSoftObjectPath WorldSoftObject(WorldToSave);
+
+		// Delete map file if it already exists
+		FString PackageFilename;
+		FPackageName::TryConvertLongPackageNameToFilename( WorldSoftObject.GetLongPackageName(), PackageFilename, FPackageName::GetMapPackageExtension() );
+
+		IFileManager::Get().Delete(*PackageFilename, /*RequireExists=*/ false, /*EvenReadOnly=*/ true, /*Quiet=*/ true);
+
+		// Add RF_Standalone flag to properly save the completed world
+		WorldToSave->SetFlags(RF_Standalone);
+
+		UEditorLoadingAndSavingUtils::SaveMap(WorldToSave, WorldSoftObject.GetLongPackageName() );
+
+		// Clear RF_Standalone from flag to properly delete and garbage collect the completed world
+		if(!bHasStandaloneFlag)
+		{
+			WorldToSave->ClearFlags(RF_Standalone);
+		}
+
+		WorldToSave->GetOutermost()->SetDirtyFlag(false);
+	}
+
+	TArray<AActor*> MoveActorsToLevel(const TArray<AActor*>& ActorsToMove, ULevel* DestLevel, bool bDuplicate)
+	{
+		if(DestLevel == nullptr || ActorsToMove.Num() == 0)
+		{
+			return TArray<AActor*>();
+		}
+
+		UWorld* OwningWorld = DestLevel->OwningWorld;
+
+		// Backup the current contents of the clipboard string as we'll be using cut/paste features to move actors
+		// between levels and this will trample over the clipboard data.
+		FString OriginalClipboardContent;
+		FPlatformApplicationMisc::ClipboardPaste(OriginalClipboardContent);
+
+		TMap<FSoftObjectPath, FSoftObjectPath> ActorPathMapping;
+		GEditor->SelectNone(false, true, false);
+
+		USelection* ActorSelection = GEditor->GetSelectedActors();
+		ActorSelection->BeginBatchSelectOperation();
+		for (AActor* Actor : ActorsToMove)
+		{
+			ActorPathMapping.Add(FSoftObjectPath(Actor), FSoftObjectPath());
+			GEditor->SelectActor(Actor, true, false);
+		}
+		ActorSelection->EndBatchSelectOperation(false);
+
+		if(GEditor->GetSelectedActorCount() == 0)
+		{
+			return TArray<AActor*>();
+		}
+
+		// Cache the old level
+		ULevel* OldCurrentLevel = OwningWorld->GetCurrentLevel();
+
+		// If we are moving the actors, cut them to remove them from the existing level
+		const bool bShoudCut = !bDuplicate;
+		const bool bIsMove = bShoudCut;
+		GEditor->CopySelectedActorsToClipboard(OwningWorld, bShoudCut, bIsMove, /*bWarnAboutReferences =*/ false);
+
+		UEditorLevelUtils::SetLevelVisibility(DestLevel, true, false);
+
+		// Scope this so that Actors that have been pasted will have their final levels set before doing the actor mapping
+		{
+			// Set the new level and force it visible while we do the paste
+			FLevelPartitionOperationScope LevelPartitionScope(DestLevel);
+			OwningWorld->SetCurrentLevel(LevelPartitionScope.GetLevel());
+
+			//const bool bDuplicate = false;
+			const bool bOffsetLocations = false;
+			const bool bWarnIfHidden = false;
+			GEditor->edactPasteSelected(OwningWorld, bDuplicate, bOffsetLocations, bWarnIfHidden);
+
+			// Restore the original current level
+			OwningWorld->SetCurrentLevel(OldCurrentLevel);
+		}
+
+		TArray<AActor*> NewActors;
+		NewActors.Reserve(GEditor->GetSelectedActorCount());
+
+		// Build a remapping of old to new names so we can do a fixup
+		for (FSelectionIterator It(GEditor->GetSelectedActorIterator()); It; ++It)
+		{
+			AActor* Actor = static_cast<AActor*>(*It);
+			if(!Actor)
+			{
+				continue;
+			}
+
+			NewActors.Add(Actor);
+			FSoftObjectPath NewPath = FSoftObjectPath(Actor);
+
+			bool bFoundMatch = false;
+
+			// First try exact match
+			for (TPair<FSoftObjectPath, FSoftObjectPath>& Pair : ActorPathMapping)
+			{
+				if (Pair.Value.IsNull() && NewPath.GetSubPathString() == Pair.Key.GetSubPathString())
+				{
+					bFoundMatch = true;
+					Pair.Value = NewPath;
+					break;
+				}
+			}
+
+			if (!bFoundMatch)
+			{
+				// Remove numbers from end as it may have had to add some to disambiguate
+				FString PartialPath = NewPath.GetSubPathString();
+				int32 IgnoreNumber;
+				FActorLabelUtilities::SplitActorLabel(PartialPath, IgnoreNumber);
+
+				for (TPair<FSoftObjectPath, FSoftObjectPath>& Pair : ActorPathMapping)
+				{
+					if (Pair.Value.IsNull())
+					{
+						FString KeyPartialPath = Pair.Key.GetSubPathString();
+						FActorLabelUtilities::SplitActorLabel(KeyPartialPath, IgnoreNumber);
+						if (PartialPath == KeyPartialPath)
+						{
+							bFoundMatch = true;
+							Pair.Value = NewPath;
+							break;
+						}
+					}
+				}
+			}
+
+			if (!bFoundMatch)
+			{
+				UE_LOG(LogDatasmithImport, Error, TEXT("Cannot find remapping for moved actor ID %s, any soft references pointing to it will be broken!"), *Actor->GetPathName());
+			}
+		}
+
+		FAssetToolsModule& AssetToolsModule = FModuleManager::GetModuleChecked<FAssetToolsModule>(TEXT("AssetTools"));
+		TArray<FAssetRenameData> RenameData;
+
+		for (TPair<FSoftObjectPath, FSoftObjectPath>& Pair : ActorPathMapping)
+		{
+			if (Pair.Value.IsValid())
+			{
+				RenameData.Add(FAssetRenameData(Pair.Key, Pair.Value, true));
+			}
+		}
+
+		if (RenameData.Num() > 0)
+		{
+			AssetToolsModule.Get().RenameAssets(RenameData);
+		}
+
+		// Restore the original clipboard contents
+		FPlatformApplicationMisc::ClipboardCopy(*OriginalClipboardContent);
+
+		return NewActors;
+	}
+
+	void MoveActorsToLevel(const TArray<AActor*>& ActorsToMove, ULevel* DestLevel, TMap<FName,TSoftObjectPtr<AActor>>& ActorsMap, bool bDuplicate)
+	{
+		if(ActorsToMove.Num() > 0)
+		{
+			UWorld* PrevGWorld = GWorld;
+			GWorld = DestLevel->OwningWorld;
+
+			// Cache Destination flags
+			EObjectFlags DestLevelFlags = DestLevel->GetFlags();
+			EObjectFlags DestWorldFlags = DestLevel->GetOuter()->GetFlags();
+			EObjectFlags DestPackageFlags = DestLevel->GetOutermost()->GetFlags();
+
+			TArray<AActor*> NewActors = MoveActorsToLevel( ActorsToMove, DestLevel, bDuplicate);
+			printf(">>> %d", NewActors.Num());
+
+			GWorld = PrevGWorld;
+
+			// Update map of related actors with new actors
+			UDatasmithContentBlueprintLibrary* DatasmithContentLibrary = Cast< UDatasmithContentBlueprintLibrary >( UDatasmithContentBlueprintLibrary::StaticClass()->GetDefaultObject() );
+
+			for(AActor* Actor : DestLevel->Actors)
+			{
+				const FString DatasmithUniqueId = DatasmithContentLibrary->GetDatasmithUserDataValueForKey( Actor, UDatasmithAssetUserData::UniqueIdMetaDataKey );
+				if(DatasmithUniqueId.Len() > 0)
+				{
+					if(TSoftObjectPtr<AActor>* SoftObjectPtr = ActorsMap.Find(FName(*DatasmithUniqueId)))
+					{
+						*SoftObjectPtr = Actor;
+					}
+				}
+			}
+
+			// Restore Destination flags
+			DestLevel->SetFlags(DestLevelFlags);
+			DestLevel->GetOuter()->SetFlags(DestWorldFlags);
+			DestLevel->GetOutermost()->SetFlags(DestPackageFlags);
+		}
+	}
 }
 
 #undef LOCTEXT_NAMESPACE
diff --git a/Engine/Plugins/Enterprise/DatasmithImporter/Source/DatasmithImporter/Private/DatasmithConsumer.h b/Engine/Plugins/Enterprise/DatasmithImporter/Source/DatasmithImporter/Private/DatasmithConsumer.h
--- a/Engine/Plugins/Enterprise/DatasmithImporter/Source/DatasmithImporter/Private/DatasmithConsumer.h
+++ b/Engine/Plugins/Enterprise/DatasmithImporter/Source/DatasmithImporter/Private/DatasmithConsumer.h
@@ -10,11 +10,17 @@
 #include "DataprepContentConsumer.h"
 
 #include "UObject/SoftObjectPath.h"
+#include "UObject/StrongObjectPtr.h"
 
 #include "DatasmithConsumer.generated.h"
 
+namespace DatasmithConsumerDetailsUtil
+{
+	class SLevelProperty;
+}
 class UDatasmithScene;
 class ULevel;
+class UWorld;
 
 UCLASS(Experimental, config = EditorSettings, HideCategories = (DatasmithConsumerInternal))
 class DATASMITHIMPORTER_API UDatasmithConsumer : public UDataprepContentConsumer
@@ -22,24 +28,32 @@ class DATASMITHIMPORTER_API UDatasmithConsumer : public UDataprepContentConsumer
 	GENERATED_BODY()
 
 public:
-	UDatasmithConsumer()
-		: DatasmithScene(nullptr)
-		, PreviousCurrentLevel(nullptr)
-	{
-	}
+	UDatasmithConsumer();
+
+	// UObject interface
+	virtual void PostLoad() override;
+	virtual void PostInitProperties() override;
+	// End of UObject interface
 
 	UPROPERTY( BlueprintReadOnly, Category = DatasmithConsumerInternal, DuplicateTransient )
 	TSoftObjectPtr<UDatasmithScene> DatasmithScene;
 
 	/** Stores the level used on the last call to UDatasmithConsumer::Run */
 	UPROPERTY( BlueprintReadOnly, Category = DatasmithConsumerInternal )
-	FString LastLevelName;
+	FString UniqueID;
+
+
+	/** Marker to tag levels, assets and actors processed by the consumer */
+	static const FString ConsumerMarkerID;
 
 	// Begin UDataprepContentConsumer overrides
 	virtual const FText& GetLabel() const override;
 	virtual const FText& GetDescription() const override;
 
 protected:
+	/** Getter/setter on consumer's output level */
+	FString GetOutputLevelPath() const { return OutputLevelSoftObject.GetAssetPathString(); }
+	bool SetOutputLevel(const FString& LevelName);
 
 	virtual bool SetLevelNameImplementation(const FString& InLevelName, FText& OutFailureReason, const bool bIsAutomated) override;
 	virtual bool SetTargetContentFolderImplementation(const FString& InTargetContentFolder, FText& OutFailureReason, const bool bIsAutomated) override;
@@ -51,23 +65,55 @@ protected:
 
 private:
 	/** Temporary code to work with UDataprepContentConsumer */
-	bool BuildContexts( UWorld* ImportWorld );
-
-	/** Returns the level associated to LevelName if it exists */
-	ULevel* FindLevel( const FString& InLevelName );
+	bool BuildContexts();
 
 	/** Move assets if destination package path has changed since last call to UDatasmithConsumer::Run */
 	void UpdateScene();
 
-	/** Move level if destination level's name has changed since last call to UDatasmithConsumer::Run */
-	void MoveLevel();
+	/** Process actors which have been marked with a UDataprepConsumerUserData */
+	void ApplySubLevelDirective();
+
+	/** Find or add a new level to the target world */
+	ULevel* FindOrAddLevel(const FString& InLevelName);
+
+	/**
+	 * Check if output directives set on assets or actors can be honored
+	 * If some issues are identified, the user is asked to proceed with execution of consumer or not.
+	 * @return True if no issue or user has accepted to continue with execution of consumer
+	 * If in command mode, the default behavior is to continue.
+	 */
+	bool CheckOutputDirectives();
+
+	/** Create the world which will be used to execute the consumer */
+	bool CreateWorld();
 
-	/** Set current level to what the user specified */
-	void UpdateLevel();
+	/** Delete the world which was used to execute the consumer */
+	void ClearWorld();
+
+	/** Save all assets and levels */
+	bool FinalizeRun();
+
+	/** Returns true if a umap file can be saved in specified folder with specified name */
+	bool CanCreateLevel( const FString& RequestedFolder, const FString& RequestedName, const bool bShowDialog);
 
 private:
 	TUniquePtr< FDatasmithImportContext > ImportContextPtr;
 	TUniquePtr< FDataprepWorkReporter > ProgressTaskPtr;
 
-	ULevel* PreviousCurrentLevel;
+	/** World used by the consumer to create new output levels */
+	TStrongObjectPtr<UWorld> WorkingWorld;
+
+	/** Indicates if LevelWorld was loaded specifically to execute this consumer */
+	bool bDeleteLevelWorld;
+
+	/** Levels which were loaded in the LevelEditor's world and need to be restored */
+	TArray<FString> LevelsToRestore;
+
+	UPROPERTY()
+	FSoftObjectPath OutputLevelSoftObject;
+
+	UPROPERTY()
+	TSoftObjectPtr<ULevel> OutputLevelSoftPtr_DEPRECATED;
+
+	ULevel* PrimaryLevel;
 };
\ No newline at end of file
diff --git a/Engine/Plugins/Enterprise/DatasmithImporter/Source/DatasmithImporter/Private/DatasmithFileProducer.cpp b/Engine/Plugins/Enterprise/DatasmithImporter/Source/DatasmithImporter/Private/DatasmithFileProducer.cpp
--- a/Engine/Plugins/Enterprise/DatasmithImporter/Source/DatasmithImporter/Private/DatasmithFileProducer.cpp
+++ b/Engine/Plugins/Enterprise/DatasmithImporter/Source/DatasmithImporter/Private/DatasmithFileProducer.cpp
@@ -423,7 +423,7 @@ void UDatasmithFileProducer::PreventNameCollision()
 
 		UDatasmithContentBlueprintLibrary* DatasmithContentLibrary = Cast< UDatasmithContentBlueprintLibrary >( UDatasmithContentBlueprintLibrary::StaticClass()->GetDefaultObject() );
 
-		auto MoveAsset = [&DatasmithContentLibrary, &AssetRedirectorMap, &PackagesToCheck](UObject* Object, UPackage* NewPackage, bool bCheckPackage)
+		auto MoveAsset = [&](UObject* Object, UPackage* NewPackage, bool bCheckPackage)
 		{
 			if(Object->GetOutermost()->GetName() != NewPackage->GetName())
 			{
@@ -432,9 +432,9 @@ void UDatasmithFileProducer::PreventNameCollision()
 				Object->Rename( nullptr, NewPackage, REN_DontCreateRedirectors | REN_NonTransactional );
 
 				// Replace unique name id with hash of package path to avoid asset's name collision
-				FString DatasmithUniqueId = DatasmithContentLibrary->GetDatasmithUserDataValueForKey( Object, UDatasmithAssetUserData::UniqueIdMetaDataKey );
-				const FString ObjectPath = Object->GetPathName();
-				FString NewUniqueId = FMD5::HashBytes( reinterpret_cast<const uint8*>(*ObjectPath), ObjectPath.Len() * sizeof(TCHAR) );
+				const FString DatasmithUniqueId = DatasmithContentLibrary->GetDatasmithUserDataValueForKey( Object, UDatasmithAssetUserData::UniqueIdMetaDataKey );
+				const FString ObjectIdent = DatasmithUniqueId + DatasmithScenePtr->GetName() + Object->GetName() + Object->GetClass()->GetName();
+				FString NewUniqueId = FMD5::HashBytes( reinterpret_cast<const uint8*>(*ObjectIdent), ObjectIdent.Len() * sizeof(TCHAR) );
 
 				UDatasmithAssetUserData::SetDatasmithUserDataValueForKey(Object, *(FString(TEXT("Old")) + UDatasmithAssetUserData::UniqueIdMetaDataKey), DatasmithUniqueId );
 				UDatasmithAssetUserData::SetDatasmithUserDataValueForKey(Object, UDatasmithAssetUserData::UniqueIdMetaDataKey, NewUniqueId );
diff --git a/Engine/Plugins/Enterprise/DatasmithImporter/Source/DatasmithImporter/Private/DatasmithImporterModule.cpp b/Engine/Plugins/Enterprise/DatasmithImporter/Source/DatasmithImporter/Private/DatasmithImporterModule.cpp
--- a/Engine/Plugins/Enterprise/DatasmithImporter/Source/DatasmithImporter/Private/DatasmithImporterModule.cpp
+++ b/Engine/Plugins/Enterprise/DatasmithImporter/Source/DatasmithImporter/Private/DatasmithImporterModule.cpp
@@ -19,6 +19,7 @@
 #include "ObjectTemplates/DatasmithObjectTemplate.h"
 #include "ObjectTemplates/DatasmithStaticMeshTemplate.h"
 #include "UI/DatasmithUIManager.h"
+#include "UI/DatasmithConsumerDetails.h"
 
 #include "AssetToolsModule.h"
 #include "ContentBrowserDelegates.h"
@@ -80,6 +81,7 @@ public:
 			FPropertyEditorModule& PropertyModule = FModuleManager::LoadModuleChecked< FPropertyEditorModule >( TEXT("PropertyEditor") );
 			PropertyModule.RegisterCustomClassLayout( TEXT("DatasmithFileProducer"), FOnGetDetailCustomizationInstance::CreateStatic( &FDatasmithFileProducerDetails::MakeDetails ) );
 			PropertyModule.RegisterCustomClassLayout( TEXT("DatasmithDirProducer"), FOnGetDetailCustomizationInstance::CreateStatic( &FDatasmithDirProducerDetails::MakeDetails ) );
+			PropertyModule.RegisterCustomClassLayout( TEXT("DatasmithConsumer"), FOnGetDetailCustomizationInstance::CreateStatic( &FDatasmithConsumerDetails::MakeDetails ) );
 
 			AddDataprepMenuEntryForDatasmithSceneAsset();
 		}
@@ -106,6 +108,8 @@ public:
 			// Register the details customizer
 			FPropertyEditorModule& PropertyModule = FModuleManager::LoadModuleChecked< FPropertyEditorModule >( TEXT("PropertyEditor") );
 			PropertyModule.UnregisterCustomClassLayout( TEXT("DatasmithFileProducer") );
+			PropertyModule.UnregisterCustomClassLayout( TEXT("DatasmithDirProducer") );
+			PropertyModule.UnregisterCustomClassLayout( TEXT("DatasmithConsumer") );
 		}
 	}
 
diff --git a/Engine/Plugins/Enterprise/DatasmithImporter/Source/DatasmithImporter/Private/UI/DatasmithConsumerDetails.cpp b/Engine/Plugins/Enterprise/DatasmithImporter/Source/DatasmithImporter/Private/UI/DatasmithConsumerDetails.cpp
new file mode 100644
--- /dev/null
+++ b/Engine/Plugins/Enterprise/DatasmithImporter/Source/DatasmithImporter/Private/UI/DatasmithConsumerDetails.cpp
@@ -0,0 +1,448 @@
+// Copyright Epic Games, Inc. All Rights Reserved.
+
+#pragma once
+
+#include "UI/DatasmithConsumerDetails.h"
+
+#include "DatasmithConsumer.h"
+#include "Utility/DatasmithImporterUtils.h"
+
+#include "DetailCategoryBuilder.h"
+#include "DetailLayoutBuilder.h"
+#include "DetailWidgetRow.h"
+#include "Dialogs/DlgPickPath.h"
+#include "EditorLevelUtils.h"
+#include "Engine/Level.h"
+#include "Engine/LevelStreamingAlwaysLoaded.h"
+#include "Engine/World.h"
+#include "Factories/LevelFactory.h"
+#include "Misc/MessageDialog.h"
+#include "Misc/PackageName.h"
+#include "PropertyCustomizationHelpers.h"
+#include "ScopedTransaction.h"
+#include "Widgets/DeclarativeSyntaxSupport.h"
+#include "Widgets/Input/SEditableTextBox.h"
+#include "Widgets/SNullWidget.h"
+#include "Widgets/Text/STextBlock.h"
+
+#define LOCTEXT_NAMESPACE "DatasmithConsumer"
+
+namespace DatasmithConsumerDetailsUtil
+{
+	/** Helper class to force a widget to fill in a space. Copied from SDetailSingleItemRow.cpp */
+	class SConstrainedBox : public SCompoundWidget
+	{
+	public:
+		SLATE_BEGIN_ARGS(SConstrainedBox) {}
+		SLATE_DEFAULT_SLOT(FArguments, Content)
+			SLATE_END_ARGS()
+
+			void Construct(const FArguments& InArgs)
+		{
+			ChildSlot
+				[
+					InArgs._Content.Widget
+				];
+		}
+
+		virtual FVector2D ComputeDesiredSize(float LayoutScaleMultiplier) const override
+		{
+			// Voluntarily ridiculously large value to force the child widget to fill up the available space
+			const float MinWidthVal = 2000;
+			const FVector2D ChildSize = ChildSlot.GetWidget()->GetDesiredSize();
+			return FVector2D(FMath::Max(MinWidthVal, ChildSize.X), ChildSize.Y);
+		}
+	};
+
+	class SFolderProperty : public SCompoundWidget
+	{
+	public:
+		SLATE_BEGIN_ARGS(SFolderProperty) {}
+		SLATE_END_ARGS()
+
+	public:
+		void Construct(const FArguments& InArgs, UDatasmithConsumer* InConsumer)
+		{
+			ConsumerPtr = InConsumer;
+			FSlateFontInfo FontInfo = IDetailLayoutBuilder::GetDetailFont();
+
+			TSharedRef<SWidget> BrowseButton = PropertyCustomizationHelpers::MakeBrowseButton( FSimpleDelegate::CreateSP( this, &SFolderProperty::OnBrowseContentFolder ) );
+
+			ChildSlot
+			[
+				SNew(SHorizontalBox)
+				+ SHorizontalBox::Slot()
+				.FillWidth(1.0f)
+				.HAlign(EHorizontalAlignment::HAlign_Fill)
+				.Padding(5.0f, 2.5f, 2.0f, 2.5f)
+				[
+					// Trick to force the splitter widget to fill up the space of its parent
+					// Strongly inspired from SDetailSingleItemRow
+					SNew(SConstrainedBox)
+					[
+						SAssignNew(ContentFolderTextBox, SEditableTextBox)
+						.Font(IDetailLayoutBuilder::GetDetailFont())
+						.HintText(LOCTEXT("DataprepSlateHelper_ContentFolderHintText", "Set the content folder to save in"))
+						.IsReadOnly(false)
+						.OnTextCommitted(FOnTextCommitted::CreateSP(this, &SFolderProperty::OnTextCommitted))
+					]
+				]
+				+ SHorizontalBox::Slot()
+				.HAlign(EHorizontalAlignment::HAlign_Right)
+				[
+					BrowseButton
+				]
+			];
+
+			UpdateContentFolderText();
+		}
+
+	private:
+		void OnBrowseContentFolder()
+		{
+							   //Ask the user for the root path where they want any content to be placed
+			if( UDataprepContentConsumer* DataprepConsumer = ConsumerPtr.Get() )
+			{
+				FString Path = DataprepConsumer->GetTargetContentFolder();
+				if( Path.IsEmpty() )
+				{
+					Path = FPaths::GetPath( DataprepConsumer->GetOutermost()->GetPathName() );
+				}
+				Path += TEXT("/"); // Trailing '/' is needed to set the default path
+
+				TSharedRef<SDlgPickPath> PickContentPathDlg =
+					SNew(SDlgPickPath)
+					.Title(LOCTEXT("DataprepSlateHelper_ChooseImportRootContentPath", "Choose Location for importing the Datasmith content"))
+					.DefaultPath(FText::FromString(Path));
+
+				if ( PickContentPathDlg->ShowModal() == EAppReturnType::Ok )
+				{
+					FScopedTransaction Transaction( LOCTEXT("Consumer_SetTargetContentFolder", "Set Target Content Folder") );
+
+					FText ErrorReason;
+					if( DataprepConsumer->SetTargetContentFolder( PickContentPathDlg->GetPath().ToString(), ErrorReason ) )
+					{
+						UpdateContentFolderText();
+					}
+					else
+					{
+						Transaction.Cancel();
+						UE_LOG( LogDatasmithImport, Error, TEXT("%s"), *ErrorReason.ToString() );
+					}
+				}
+			}
+		}
+
+		void UpdateContentFolderText()
+		{
+			if(UDataprepContentConsumer* Consumer = ConsumerPtr.Get())
+			{
+				FString TargetContentFolder( Consumer->GetTargetContentFolder() );
+
+				if( TargetContentFolder.IsEmpty() )
+				{
+					TargetContentFolder = TEXT("/Content");
+				}
+				else if( TargetContentFolder.StartsWith( TEXT( "/Game" ) ) )
+				{
+					TargetContentFolder = TargetContentFolder.Replace( TEXT( "/Game" ), TEXT( "/Content" ) );
+				}
+
+				ContentFolderTextBox->SetText( FText::FromString( TargetContentFolder + TEXT( "/" ) ) );
+			}
+			else
+			{
+				ContentFolderTextBox->SetText( TAttribute<FText>() );
+			}
+		}
+
+		void OnTextCommitted( const FText& NewText, ETextCommit::Type CommitType)
+		{
+			if( UDataprepContentConsumer* DataprepConsumer = ConsumerPtr.Get() )
+			{
+				FString NewContentFolder( NewText.ToString() );
+
+				// Replace /Content/ with /Game/ since /Content is only used for display 
+				if( NewContentFolder.StartsWith( TEXT("/Content") ) )
+				{
+					NewContentFolder = NewContentFolder.Replace( TEXT( "/Content" ), TEXT( "/Game" ) );
+				}
+
+				// Remove ending '/' if applicable
+				if( !NewContentFolder.IsEmpty() && NewContentFolder[ NewContentFolder.Len()-1 ] == TEXT('/') )
+				{
+					NewContentFolder[ NewContentFolder.Len()-1 ] = 0;
+					NewContentFolder = NewContentFolder.LeftChop(1);
+				}
+
+				if(NewContentFolder != DataprepConsumer->GetTargetContentFolder())
+				{
+					FScopedTransaction Transaction( LOCTEXT("Consumer_SetTargetContentFolder", "Set Target Content Folder") );
+
+					FText ErrorReason;
+					if( !DataprepConsumer->SetTargetContentFolder( NewContentFolder, ErrorReason ) )
+					{
+						Transaction.Cancel();
+						UE_LOG( LogDatasmithImport, Error, TEXT("%s"), *ErrorReason.ToString() );
+						UpdateContentFolderText();
+					}
+				}
+			}
+		}
+
+	private:
+		TWeakObjectPtr< UDatasmithConsumer > ConsumerPtr;
+		TSharedPtr< SEditableTextBox > ContentFolderTextBox;
+	};
+
+	class SLevelProperty : public SCompoundWidget
+	{
+	public:
+		SLATE_BEGIN_ARGS(SLevelProperty) {}
+		SLATE_END_ARGS()
+
+	public:
+		void Construct(const FArguments& InArgs, UDatasmithConsumer* InConsumer)
+		{
+			ConsumerPtr = InConsumer;
+			FSlateFontInfo FontInfo = IDetailLayoutBuilder::GetDetailFont();
+
+			ChildSlot
+			[
+				SNew(SHorizontalBox)
+				+ SHorizontalBox::Slot()
+				.FillWidth(1.0f)
+				.HAlign(EHorizontalAlignment::HAlign_Fill)
+				.Padding(5.0f, 2.5f, 2.0f, 2.5f)
+				[
+					// Trick to force the splitter widget to fill up the space of its parent
+					// Strongly inspired from SDetailSingleItemRow
+					//SNew(SConstrainedBox)
+					//[
+						SAssignNew(LevelTextBox, SEditableTextBox)
+						.Font(IDetailLayoutBuilder::GetDetailFont())
+						.HintText(LOCTEXT("DataprepLevelProperty_HintText", "Set the name of the level to save in"))
+						.IsReadOnly(false)
+						.OnTextCommitted(FOnTextCommitted::CreateSP(this, &SLevelProperty::OnTextCommitted))
+					//]
+				]
+			];
+
+			UpdateLevelText();
+		}
+
+	private:
+		void UpdateLevelText()
+		{
+			if(UDatasmithConsumer* Consumer = ConsumerPtr.Get())
+			{
+				LevelTextBox->SetText( FText::FromString( Consumer->GetLevelName() ) );
+			}
+			else
+			{
+				LevelTextBox->SetText( TAttribute<FText>() );
+			}
+		}
+
+		void OnTextCommitted( const FText& NewText, ETextCommit::Type CommitType)
+		{
+			if( UDatasmithConsumer* DataprepConsumer = ConsumerPtr.Get() )
+			{
+				FString NewLevelName( NewText.ToString() );
+
+				if(NewLevelName != DataprepConsumer->GetLevelName())
+				{
+					FScopedTransaction Transaction( LOCTEXT("Consumer_SetLevelName", "Set Output Level") );
+
+					FText OutReason;
+					if( !DataprepConsumer->SetLevelName( NewLevelName, OutReason ) )
+					{
+						Transaction.Cancel();
+
+						UpdateLevelText();
+
+						UE_LOG( LogDatasmithImport, Error, TEXT("Cannot create a level named %s - %s"), *NewLevelName, *OutReason.ToString() );
+					}
+				}
+			}
+		}
+
+	private:
+		TWeakObjectPtr< UDatasmithConsumer > ConsumerPtr;
+		TSharedPtr< SEditableTextBox > LevelTextBox;
+	};
+
+	class SLevelAssetProperty : public SCompoundWidget
+	{
+	public:
+		SLATE_BEGIN_ARGS(SLevelAssetProperty) {}
+		SLATE_END_ARGS()
+
+	public:
+		void Construct(const FArguments& InArgs, UDatasmithConsumer* InConsumer)
+		{
+			ConsumerPtr = InConsumer;
+			FSlateFontInfo FontInfo = IDetailLayoutBuilder::GetDetailFont();
+
+			TArray<FAssetData> AssetDataArray;
+			AssetDataArray.Add(FAssetData(ConsumerPtr->GetOutermost()));
+
+			TSharedPtr<SWidget> OutputLevelWidget = SNew(SObjectPropertyEntryBox)
+				.AllowedClass(UWorld::StaticClass())
+				.ObjectPath(this, &SLevelAssetProperty::GetAssetPath)
+				.OnShouldFilterAsset(this, &SLevelAssetProperty::OnShouldFilterAsset)
+				.OnObjectChanged(this, &SLevelAssetProperty::OnAssetSelected)
+				.OwnerAssetDataArray(AssetDataArray)
+				.AllowClear( false )
+				.DisplayUseSelected( true )
+				.DisplayBrowse( true )
+				.EnableContentPicker(true)
+				.DisplayCompactSize(true)
+				.DisplayThumbnail(false);
+
+			ChildSlot
+			[
+				SNew(SConstrainedBox)
+				[
+					OutputLevelWidget.ToSharedRef()
+				]
+			];
+		}
+
+	private:
+		void OnAssetSelected(const FAssetData& AssetData)
+		{
+			if ( UDatasmithConsumer* Consumer = ConsumerPtr.Get() )
+			{
+				UWorld* World = Cast<UWorld>(AssetData.GetAsset());
+				ensure(World);
+
+				// Check that level has been chosen or created in the right content folder
+				if(ULevel* Level = World->PersistentLevel)
+				{
+					const FString PackagePathName = FPaths::GetPath(Level->GetOutermost()->GetPathName());
+					const FString TargetContentFolder = Consumer->GetTargetContentFolder();
+					if(PackagePathName == TargetContentFolder)
+					{
+						FScopedTransaction Transaction( LOCTEXT("Consumer_SetOutputLevel", "Set Output Level") );
+						FText OutReason;
+
+						if(Consumer->SetLevelName(AssetData.AssetName.ToString(), OutReason))
+						{
+							WorldPtr = World;
+						}
+						else
+						{
+							Transaction.Cancel();
+							UE_LOG( LogDatasmithImport, Error, TEXT("Cannot create a level named %s - %s"), *AssetData.AssetName.ToString(), *OutReason.ToString() );
+						}
+					}
+					// Inform user about where the level should be created
+					else
+					{
+						const FText WarningMessage = FText::Format(LOCTEXT("LevelProperty_WrongFolder", "The new level must be created in the asset's folder, {0}."), FText::FromString(TargetContentFolder));
+						const FText DialogTitle(LOCTEXT("LevelProperty_WrongFolder_Title", "Warning - Wrong Folder"));
+
+						FMessageDialog::Open(EAppMsgType::Ok, WarningMessage, &DialogTitle);
+					}
+				}
+			}
+		}
+
+		FString GetAssetPath() const
+		{
+			//if ( UDatasmithConsumer* Consumer = ConsumerPtr.Get() )
+			//{
+			//	if(!WorldPtr.IsValid())
+			//	{
+			//		FString OutputLevelPath = Consumer->GetOutputLevelPath();
+			//	}
+
+			//	return FSoftObjectPath(WorldPtr->PersistentLevel).GetAssetPathString();
+			//}
+
+			return FString("");
+		}
+
+		bool OnShouldFilterAsset(const FAssetData& AssetData) const
+		{
+			if ( UDatasmithConsumer* Consumer = ConsumerPtr.Get() )
+			{
+				bool bValidAsset = AssetData.PackagePath.ToString() == Consumer->GetTargetContentFolder();
+				bValidAsset &= AssetData.AssetClass == UWorld::StaticClass()->GetFName();
+
+				return !bValidAsset;
+			}
+
+			return true;
+		}
+
+	private:
+		TWeakObjectPtr< UDatasmithConsumer > ConsumerPtr;
+		TWeakObjectPtr< UWorld > WorldPtr;
+	};
+}
+
+void FDatasmithConsumerDetails::CustomizeDetails(IDetailLayoutBuilder& DetailBuilder)
+{
+	TArray< TWeakObjectPtr< UObject > > Objects;
+	DetailBuilder.GetObjectsBeingCustomized( Objects );
+	check( Objects.Num() > 0 );
+
+	UDatasmithConsumer* Consumer = Cast< UDatasmithConsumer >(Objects[0].Get());
+	check( Consumer );
+
+	//FName CategoryName( TEXT("DatasmithFileProducerCustom") );
+	IDetailCategoryBuilder& ImportSettingsCategoryBuilder = DetailBuilder.EditCategory( NAME_None, FText::GetEmpty(), ECategoryPriority::Important );
+
+	TSharedRef<IPropertyHandle> TargetFolderHandle = DetailBuilder.GetProperty( TEXT("TargetContentFolder"), UDataprepContentConsumer::StaticClass() );
+	TargetFolderHandle->MarkHiddenByCustomization();
+
+	TSharedRef<IPropertyHandle> LevelNameHandle = DetailBuilder.GetProperty( TEXT("LevelName"), UDataprepContentConsumer::StaticClass() );
+	LevelNameHandle->MarkHiddenByCustomization();
+
+	TSharedRef<IPropertyHandle> LevelPathHandle = DetailBuilder.GetProperty( TEXT("OutputLevelPath"), UDatasmithConsumer::StaticClass() );
+	LevelPathHandle->MarkHiddenByCustomization();
+
+	{
+		using namespace DatasmithConsumerDetailsUtil;
+
+		FText PropertyName = FText::FromString( TEXT( "Folder" ) );
+		FDetailWidgetRow& CustomRow = ImportSettingsCategoryBuilder.AddCustomRow( PropertyName );
+
+		CustomRow.NameContent()
+		[
+			SNew( STextBlock )
+			.Text( PropertyName )
+			.Font( IDetailLayoutBuilder::GetDetailFont() )
+		];
+
+		CustomRow.ValueContent()
+		.MinDesiredWidth( 2000.0f )
+		[
+			SNew( SFolderProperty, Consumer )
+		];
+	}
+
+	{
+		using namespace DatasmithConsumerDetailsUtil;
+
+		FText PropertyName = FText::FromString( TEXT( "Level Name" ) );
+		FDetailWidgetRow& CustomRow = ImportSettingsCategoryBuilder.AddCustomRow( PropertyName );
+
+		CustomRow.NameContent()
+		[
+			SNew( STextBlock )
+			.Text( PropertyName )
+			.Font( IDetailLayoutBuilder::GetDetailFont() )
+		];
+
+		CustomRow.ValueContent()
+		.MinDesiredWidth( 2000.0f )
+		[
+			SNew( SLevelProperty, Consumer )
+		];
+	}
+}
+
+#undef LOCTEXT_NAMESPACE
\ No newline at end of file
diff --git a/Engine/Plugins/Enterprise/DatasmithImporter/Source/DatasmithImporter/Private/UI/DatasmithConsumerDetails.h b/Engine/Plugins/Enterprise/DatasmithImporter/Source/DatasmithImporter/Private/UI/DatasmithConsumerDetails.h
new file mode 100644
--- /dev/null
+++ b/Engine/Plugins/Enterprise/DatasmithImporter/Source/DatasmithImporter/Private/UI/DatasmithConsumerDetails.h
@@ -0,0 +1,20 @@
+// Copyright Epic Games, Inc. All Rights Reserved.
+
+#pragma once
+
+#include "CoreMinimal.h"
+
+#include "IDetailCustomization.h"
+
+class IDetailLayoutBuilder;
+
+// Customization of the details of the Datasmith Consumer for the data prep editor.
+class DATASMITHIMPORTER_API FDatasmithConsumerDetails : public IDetailCustomization
+{
+public:
+	static TSharedRef< IDetailCustomization > MakeDetails() { return MakeShared<FDatasmithConsumerDetails>(); };
+
+	/** Called when details should be customized */
+	virtual void CustomizeDetails( IDetailLayoutBuilder& DetailBuilder ) override;
+};
+
diff --git a/Engine/Plugins/Enterprise/DatasmithImporter/Source/DatasmithImporter/Private/Utility/DatasmithImporterUtils.cpp b/Engine/Plugins/Enterprise/DatasmithImporter/Source/DatasmithImporter/Private/Utility/DatasmithImporterUtils.cpp
--- a/Engine/Plugins/Enterprise/DatasmithImporter/Source/DatasmithImporter/Private/Utility/DatasmithImporterUtils.cpp
+++ b/Engine/Plugins/Enterprise/DatasmithImporter/Source/DatasmithImporter/Private/Utility/DatasmithImporterUtils.cpp
@@ -403,6 +403,37 @@ void FDatasmithImporterUtils::AddUniqueLayersToWorld(UWorld* World, const TSet<
 }
 
 bool FDatasmithImporterUtils::CanCreateAsset(const FString& AssetPathName, const UClass* AssetClass, FText& OutFailReason)
+{
+	switch(CanCreateAsset(AssetPathName, AssetClass))
+	{
+		case EAssetCreationStatus::CS_HasRedirector:
+		{
+			OutFailReason = FText::Format(LOCTEXT("FoundRedirectionForAsset", "Found redirection for asset {0}. Skipping this asset ..."), FText::FromString(AssetPathName));
+			return false;
+		}
+
+		case EAssetCreationStatus::CS_ClassMismatch:
+		{
+			IAssetRegistry& AssetRegistry = FModuleManager::LoadModuleChecked<FAssetRegistryModule>(TEXT("AssetRegistry")).Get();
+			const FAssetData AssetData = AssetRegistry.GetAssetByObjectPath(*AssetPathName);
+
+			const FString FoundClassName(AssetData.GetClass()->GetFName().ToString());
+			const FString ExpectedClassName(AssetClass->GetFName().ToString());
+			OutFailReason = FText::Format(LOCTEXT("AssetClassMismatch", "Found asset {0} of class {1} instead of class {2}. Skipping this asset ..."), FText::FromString(AssetPathName), FText::FromString(FoundClassName), FText::FromString(ExpectedClassName) );
+			return false;
+		}
+
+		case EAssetCreationStatus::CS_CanCreate:
+		default:
+		{
+			break;
+		}
+	}
+
+	return true;
+}
+
+FDatasmithImporterUtils::EAssetCreationStatus FDatasmithImporterUtils::CanCreateAsset(const FString& AssetPathName, const UClass* AssetClass)
 {
 	IAssetRegistry& AssetRegistry = FModuleManager::LoadModuleChecked<FAssetRegistryModule>(TEXT("AssetRegistry")).Get();
 
@@ -411,25 +442,21 @@ bool FDatasmithImporterUtils::CanCreateAsset(const FString& AssetPathName, const
 	// Asset does not exist yet. Safe to import
 	if (!AssetData.IsValid())
 	{
-		return true;
+		return EAssetCreationStatus::CS_CanCreate;
 	}
 
 	// Warn and skip import of asset since it is an object redirection
 	if (AssetData.IsRedirector())
 	{
-		OutFailReason = FText::Format(LOCTEXT("FoundRedirectionForAsset", "Found redirection for asset {0}. Skipping this asset ..."), FText::FromString(AssetPathName));
-		return false;
+		return EAssetCreationStatus::CS_HasRedirector;
 	}
 	// Warn and skip re-import of asset since it is not of the expected class
 	else if (!AssetData.GetClass()->IsChildOf(AssetClass))
 	{
-		const FString FoundClassName(AssetData.GetClass()->GetFName().ToString());
-		const FString ExpectedClassName(AssetClass->GetFName().ToString());
-		OutFailReason = FText::Format(LOCTEXT("AssetClassMismatch", "Found asset {0} of class {1} instead of class {2}. Skipping this asset ..."), FText::FromString(AssetPathName), FText::FromString(FoundClassName), FText::FromString(ExpectedClassName) );
-		return false;
+		return EAssetCreationStatus::CS_ClassMismatch;
 	}
 
-	return true;
+	return EAssetCreationStatus::CS_CanCreate;
 }
 
 UDatasmithScene* FDatasmithImporterUtils::FindDatasmithSceneForAsset( UObject* Asset )
diff --git a/Engine/Plugins/Enterprise/DatasmithImporter/Source/DatasmithImporter/Public/Utility/DatasmithImporterUtils.h b/Engine/Plugins/Enterprise/DatasmithImporter/Source/DatasmithImporter/Public/Utility/DatasmithImporterUtils.h
--- a/Engine/Plugins/Enterprise/DatasmithImporter/Source/DatasmithImporter/Public/Utility/DatasmithImporterUtils.h
+++ b/Engine/Plugins/Enterprise/DatasmithImporter/Source/DatasmithImporter/Public/Utility/DatasmithImporterUtils.h
@@ -18,7 +18,14 @@ DECLARE_LOG_CATEGORY_EXTERN( LogDatasmithImport, Log, All );
 class DATASMITHIMPORTER_API FDatasmithImporterUtils
 {
 public:
-	/** Loads an IDatasmithScene from a UDatasmithScene */
+	enum class EAssetCreationStatus : uint8
+	{
+		CS_CanCreate,
+		CS_HasRedirector,
+		CS_ClassMismatch,
+	};
+
+		/** Loads an IDatasmithScene from a UDatasmithScene */
 	static TSharedPtr< IDatasmithScene > LoadDatasmithScene( UDatasmithScene* DatasmithSceneAsset );
 
 	/** Saves an IDatasmithScene into a UDatasmithScene */
@@ -110,6 +117,25 @@ public:
 	 */
 	static bool CanCreateAsset(const FString& AssetPathName, const UClass* AssetClass, FText& OutFailReason);
 
+	/**
+	 * @param AssetPathName		Full path name of the asset to create
+	 * @param AssetClass		Class of the asset to create
+	 * Returns a state from the creation enumeration, EAssetCreationStatus.
+	 * Given a path and a class, check if an existing asset with a different class
+	 * would not prevent the creation of such asset.
+	 */
+	static EAssetCreationStatus CanCreateAsset(const FString& AssetPathName, const UClass* AssetClass);
+
+	/**
+	 * @param AssetPathName		Full path name of the asset to create
+	 * Calls CanCreateAsset(const FString&, const UClass*) with the instantiating class
+	 */
+	template< class ObjectType >
+	static EAssetCreationStatus CanCreateAsset(const FString& AssetPathName)
+	{
+		return CanCreateAsset(AssetPathName, ObjectType::StaticClass());
+	}
+
 	/**
 	 * Finds the UDatasmithScene for which the Asset belongs to.
 	 */
