Correction to CL #11569974

diff --git a/Engine/Plugins/Enterprise/DatasmithImporter/Source/DatasmithImporter/DatasmithImporter.Build.cs b/Engine/Plugins/Enterprise/DatasmithImporter/Source/DatasmithImporter/DatasmithImporter.Build.cs
--- a/Engine/Plugins/Enterprise/DatasmithImporter/Source/DatasmithImporter/DatasmithImporter.Build.cs
+++ b/Engine/Plugins/Enterprise/DatasmithImporter/Source/DatasmithImporter/DatasmithImporter.Build.cs
@@ -7,8 +7,6 @@ namespace UnrealBuildTool.Rules
 		public DatasmithImporter(ReadOnlyTargetRules Target)
 			: base(Target)
 		{
-			OptimizeCode = CodeOptimization.InShippingBuildsOnly;
-
 			PrivateDependencyModuleNames.AddRange(
 				new string[]
 				{
diff --git a/Engine/Plugins/Enterprise/DatasmithImporter/Source/DatasmithImporter/Private/DatasmithConsumer.cpp b/Engine/Plugins/Enterprise/DatasmithImporter/Source/DatasmithImporter/Private/DatasmithConsumer.cpp
--- a/Engine/Plugins/Enterprise/DatasmithImporter/Source/DatasmithImporter/Private/DatasmithConsumer.cpp
+++ b/Engine/Plugins/Enterprise/DatasmithImporter/Source/DatasmithImporter/Private/DatasmithConsumer.cpp
@@ -251,13 +251,6 @@ void UDatasmithConsumer::PostLoad()
 			LevelName = GetOuter()->GetName() + TEXT("_Map");
 		}
 
-		if(OutputLevelSoftPtr_DEPRECATED.ToString().Len() > 0)
-		{
-			OutputLevelSoftObject = OutputLevelSoftPtr_DEPRECATED.ToSoftObjectPath();
-
-			bMarkDirty = true;
-		}
-
 		if(OutputLevelSoftObject.GetAssetPathString().Len() == 0)
 		{
 			OutputLevelSoftObject = FSoftObjectPath(FPaths::Combine(TargetContentFolder, LevelName) + "." + LevelName);
diff --git a/Engine/Plugins/Enterprise/DatasmithImporter/Source/DatasmithImporter/Private/DatasmithConsumer.h b/Engine/Plugins/Enterprise/DatasmithImporter/Source/DatasmithImporter/Private/DatasmithConsumer.h
--- a/Engine/Plugins/Enterprise/DatasmithImporter/Source/DatasmithImporter/Private/DatasmithConsumer.h
+++ b/Engine/Plugins/Enterprise/DatasmithImporter/Source/DatasmithImporter/Private/DatasmithConsumer.h
@@ -112,8 +112,5 @@ private:
 	UPROPERTY()
 	FSoftObjectPath OutputLevelSoftObject;
 
-	UPROPERTY()
-	TSoftObjectPtr<ULevel> OutputLevelSoftPtr_DEPRECATED;
-
 	ULevel* PrimaryLevel;
 };
\ No newline at end of file
