Fixed localization conflicts

#rnx

diff --git a/Engine/Plugins/Enterprise/DataprepEditor/Source/DataprepEditor/Private/Widgets/DataprepGraph/SDataprepGraphEditor.cpp b/Engine/Plugins/Enterprise/DataprepEditor/Source/DataprepEditor/Private/Widgets/DataprepGraph/SDataprepGraphEditor.cpp
--- a/Engine/Plugins/Enterprise/DataprepEditor/Source/DataprepEditor/Private/Widgets/DataprepGraph/SDataprepGraphEditor.cpp
+++ b/Engine/Plugins/Enterprise/DataprepEditor/Source/DataprepEditor/Private/Widgets/DataprepGraph/SDataprepGraphEditor.cpp
@@ -624,7 +624,7 @@ void SDataprepGraphEditor::PasteNodes()
 	FPlatformApplicationMisc::ClipboardPaste(TextToImport);
 
 	// Create temporary graph
-	FName UniqueGraphName = MakeUniqueObjectName( GetTransientPackage(), UWorld::StaticClass(), FName( *(LOCTEXT("DataprepGraph", "TempGraph").ToString()) ) );
+	FName UniqueGraphName = MakeUniqueObjectName( GetTransientPackage(), UWorld::StaticClass(), FName( *(LOCTEXT("DataprepTempGraph", "TempGraph").ToString()) ) );
 	TStrongObjectPtr<UDataprepGraph> DataprepGraph = TStrongObjectPtr<UDataprepGraph>( NewObject< UDataprepGraph >(GetTransientPackage(), UniqueGraphName) );
 	DataprepGraph->Schema = UDataprepGraphSchema::StaticClass();
 
diff --git a/Engine/Plugins/Enterprise/DatasmithImporter/Source/DatasmithImporter/Private/DatasmithFileProducer.cpp b/Engine/Plugins/Enterprise/DatasmithImporter/Source/DatasmithImporter/Private/DatasmithFileProducer.cpp
--- a/Engine/Plugins/Enterprise/DatasmithImporter/Source/DatasmithImporter/Private/DatasmithFileProducer.cpp
+++ b/Engine/Plugins/Enterprise/DatasmithImporter/Source/DatasmithImporter/Private/DatasmithFileProducer.cpp
@@ -825,7 +825,7 @@ bool UDatasmithDirProducer::Execute(TArray< TWeakObjectPtr< UObject > >& OutAsse
 
 		if( !FileProducer->Produce( Context, OutAssets ) )
 		{
-			FText ErrorReport = FText::Format( LOCTEXT( "DatasmithDirProducer_Failed", "Failed to load {0} ..."), FText::FromString( FileName ) );
+			FText ErrorReport = FText::Format( LOCTEXT( "DatasmithDirProducer_FailedLoad", "Failed to load {0} ..."), FText::FromString( FileName ) );
 			LogError( ErrorReport );
 		}
 	}
@@ -957,7 +957,7 @@ bool UDatasmithDirProducer::ImportAsPlmXml(UPackage* RootPackage, TArray<TWeakOb
 
 	if (!FFileHelper::SaveStringToFile(Buffer, *PlmXmlFileName, FFileHelper::EEncodingOptions::ForceUTF8WithoutBOM))
 	{
-		FText ErrorReport = FText::Format(LOCTEXT("DatasmithDirProducer_Failed", "Failed to create PlmXml file, {0}, for parallel loading ..."), FText::FromString(PlmXmlFileName));
+		FText ErrorReport = FText::Format(LOCTEXT("DatasmithPlmXmlProducer_FailedCreate", "Failed to create PlmXml file, {0}, for parallel loading ..."), FText::FromString(PlmXmlFileName));
 		LogError(ErrorReport);
 		return false;
 	}
@@ -973,7 +973,7 @@ bool UDatasmithDirProducer::ImportAsPlmXml(UPackage* RootPackage, TArray<TWeakOb
 
 	if (!FileProducer->Produce(Context, OutAssets))
 	{
-		FText ErrorReport = FText::Format(LOCTEXT("DatasmithDirProducer_Failed", "Failed to produce assets with PlmXml file, {0}, for parallel loading ..."), FText::FromString(PlmXmlFileName));
+		FText ErrorReport = FText::Format(LOCTEXT("DatasmithPlmXmlProducer_FailedProduce", "Failed to produce assets with PlmXml file, {0}, for parallel loading ..."), FText::FromString(PlmXmlFileName));
 		LogError(ErrorReport);
 		return false;
 	}
diff --git a/Engine/Plugins/Experimental/DigitalMultiplex/DMXEngine/Source/DMXEditor/Private/Widgets/SDMXEntityList.cpp b/Engine/Plugins/Experimental/DigitalMultiplex/DMXEngine/Source/DMXEditor/Private/Widgets/SDMXEntityList.cpp
--- a/Engine/Plugins/Experimental/DigitalMultiplex/DMXEngine/Source/DMXEditor/Private/Widgets/SDMXEntityList.cpp
+++ b/Engine/Plugins/Experimental/DigitalMultiplex/DMXEngine/Source/DMXEditor/Private/Widgets/SDMXEntityList.cpp
@@ -579,7 +579,7 @@ void FDMXEntityDragDropOperation::HoverTargetChanged()
 			if (bValidDropTarget)
 			{
 				SetFeedbackMessageOK(FText::Format(
-					LOCTEXT("ReorderAndSetProperty", "Reorder {0} before '{1}'\nSet {2} = '{3}'"),
+					LOCTEXT("ReorderBeforeAndSetProperty", "Reorder {0} before '{1}'\nSet {2} = '{3}'"),
 					DraggedLabel,
 					HoveredEntity->GetDisplayName(),
 					PropertyChangeName,
@@ -601,7 +601,7 @@ void FDMXEntityDragDropOperation::HoverTargetChanged()
 		{
 			// Reordering between entities of same category is fine
 			SetFeedbackMessageOK(FText::Format(
-				LOCTEXT("ReorderEntities", "Reorder {0} before '{1}'"),
+				LOCTEXT("ReorderBeforeOther", "Reorder {0} before '{1}'"),
 				DraggedLabel,
 				HoveredEntity->GetDisplayName()
 			));
diff --git a/Engine/Plugins/Experimental/HairStrands/Source/HairStrandsEditor/Private/GroomImportOptionsWindow.cpp b/Engine/Plugins/Experimental/HairStrands/Source/HairStrandsEditor/Private/GroomImportOptionsWindow.cpp
--- a/Engine/Plugins/Experimental/HairStrands/Source/HairStrandsEditor/Private/GroomImportOptionsWindow.cpp
+++ b/Engine/Plugins/Experimental/HairStrands/Source/HairStrandsEditor/Private/GroomImportOptionsWindow.cpp
@@ -159,12 +159,12 @@ TSharedPtr<SGroomImportOptionsWindow> DisplayOptions(UGroomImportOptions* Import
 
 TSharedPtr<SGroomImportOptionsWindow> SGroomImportOptionsWindow::DisplayImportOptions(UGroomImportOptions* ImportOptions, const FString& FilePath)
 {
-	return DisplayOptions(ImportOptions, FilePath, EGroomOptionsVisibility::All, LOCTEXT("WindowTitle", "Groom Import Options"), LOCTEXT("Import", "Import"));
+	return DisplayOptions(ImportOptions, FilePath, EGroomOptionsVisibility::All, LOCTEXT("ImportOptions", "Groom Import Options"), LOCTEXT("Import", "Import"));
 }
 
 TSharedPtr<SGroomImportOptionsWindow> SGroomImportOptionsWindow::DisplayRebuildOptions(UGroomImportOptions* ImportOptions, const FString& FilePath)
 {
-	return DisplayOptions(ImportOptions, FilePath, EGroomOptionsVisibility::BuildOptions, LOCTEXT("WindowTitle", "Groom Build Options"), LOCTEXT("Build", "Build"));
+	return DisplayOptions(ImportOptions, FilePath, EGroomOptionsVisibility::BuildOptions, LOCTEXT("RebuildOptions", "Groom Build Options"), LOCTEXT("Build", "Build"));
 }
 
 
diff --git a/Engine/Plugins/FX/Niagara/Source/NiagaraEditor/Private/Customizations/NiagaraTypeCustomizations.cpp b/Engine/Plugins/FX/Niagara/Source/NiagaraEditor/Private/Customizations/NiagaraTypeCustomizations.cpp
--- a/Engine/Plugins/FX/Niagara/Source/NiagaraEditor/Private/Customizations/NiagaraTypeCustomizations.cpp
+++ b/Engine/Plugins/FX/Niagara/Source/NiagaraEditor/Private/Customizations/NiagaraTypeCustomizations.cpp
@@ -418,7 +418,7 @@ void FNiagaraUserParameterBindingCustomization::OnActionSelected(const TArray< T
 
 void FNiagaraUserParameterBindingCustomization::ChangeSource(FName InVarName)
 {
-	FScopedTransaction Transaction(FText::Format(LOCTEXT("ChangeParameterSource", " Change User Parameter Source to \"{0}\" "), FText::FromName(InVarName)));
+	FScopedTransaction Transaction(FText::Format(LOCTEXT("ChangeUserParameterSource", " Change User Parameter Source to \"{0}\" "), FText::FromName(InVarName)));
 	TArray<UObject*> Objects;
 	PropertyHandle->GetOuterObjects(Objects);
 	for (UObject* Obj : Objects)
@@ -627,7 +627,7 @@ void FNiagaraDataInterfaceBindingCustomization::OnActionSelected(const TArray< T
 
 void FNiagaraDataInterfaceBindingCustomization::ChangeSource(FName InVarName)
 {
-	FScopedTransaction Transaction(FText::Format(LOCTEXT("ChangeParameterSource", " Change Data Interface Source to \"{0}\" "), FText::FromName(InVarName)));
+	FScopedTransaction Transaction(FText::Format(LOCTEXT("ChangeDataParameterSource", " Change Data Interface Source to \"{0}\" "), FText::FromName(InVarName)));
 	TArray<UObject*> Objects;
 	PropertyHandle->GetOuterObjects(Objects);
 	for (UObject* Obj : Objects)
diff --git a/Engine/Plugins/VirtualProduction/TimedDataMonitor/Source/TimedDataMonitor/Private/TimedDataMonitorCalibration.cpp b/Engine/Plugins/VirtualProduction/TimedDataMonitor/Source/TimedDataMonitor/Private/TimedDataMonitorCalibration.cpp
--- a/Engine/Plugins/VirtualProduction/TimedDataMonitor/Source/TimedDataMonitor/Private/TimedDataMonitorCalibration.cpp
+++ b/Engine/Plugins/VirtualProduction/TimedDataMonitor/Source/TimedDataMonitor/Private/TimedDataMonitorCalibration.cpp
@@ -780,6 +780,3 @@ FTimedDataMonitorTimeCorrectionResult FTimedDataMonitorCalibration::ApplyTimeCor
 	Result.ReturnCode = ETimedDataMonitorTimeCorrectionReturnCode::Retry_IncreaseBufferSize;
 	return Result;
 }
-
-
-#undef LOCTEXT_NAMESPACE
diff --git a/Engine/Plugins/VirtualProduction/TimedDataMonitor/Source/TimedDataMonitorEditor/Private/STimingDiagramWidget.cpp b/Engine/Plugins/VirtualProduction/TimedDataMonitor/Source/TimedDataMonitorEditor/Private/STimingDiagramWidget.cpp
--- a/Engine/Plugins/VirtualProduction/TimedDataMonitor/Source/TimedDataMonitorEditor/Private/STimingDiagramWidget.cpp
+++ b/Engine/Plugins/VirtualProduction/TimedDataMonitor/Source/TimedDataMonitorEditor/Private/STimingDiagramWidget.cpp
@@ -317,7 +317,7 @@ FText STimingDiagramWidget::GetTooltipText() const
 		const float DistanceToNewestSampleSigma = TimedDataMonitorSubsystem->GetInputEvaluationDistanceToNewestSampleStandardDeviation(InputIdentifier);
 		const float DistanceToOldestSampleSigma = TimedDataMonitorSubsystem->GetInputEvaluationDistanceToOldestSampleStandardDeviation(InputIdentifier);
 
-		return FText::Format(LOCTEXT("TimingDiagramTooltip", "Distance to newest:\nMean: {0}\nSigma: {1}\n\nDistance to oldest:\nMean: {2}\nSigma: {3}")
+		return FText::Format(LOCTEXT("TimingDiagramTooltip_DistanceToNewest", "Distance to newest:\nMean: {0}\nSigma: {1}\n\nDistance to oldest:\nMean: {2}\nSigma: {3}")
 			, DistanceToNewestSampleAverage
 			, DistanceToNewestSampleSigma
 			, DistanceToOldestSampleAverage
@@ -333,7 +333,7 @@ FText STimingDiagramWidget::GetTooltipText() const
 		FTimedDataInputEvaluationData LastEvalData;
 		TimedDataMonitorSubsystem->GetChannelLastEvaluationDataStat(ChannelIdentifier, LastEvalData);
 	
-		return FText::Format(LOCTEXT("TimingDiagramTooltip", "Sample count: {0}\n\nDistance to newest:\nLast Value: {1}\nMean: {2}\nSigma: {3}\n\nDistance to oldest:\nLast Value: {4}\nMean: {5}\nSigma: {6}")
+		return FText::Format(LOCTEXT("TimingDiagramTooltip_SampleCount", "Sample count: {0}\n\nDistance to newest:\nLast Value: {1}\nMean: {2}\nSigma: {3}\n\nDistance to oldest:\nLast Value: {4}\nMean: {5}\nSigma: {6}")
 		, SampleCount
 		, LastEvalData.DistanceToNewestSampleSeconds
 		, DistanceToNewestSampleAverage
