Integrating CL# 11577455 from Dev-Enterprise
Fixing dataprep material reimport

diff --git a/Engine/Plugins/Enterprise/DatasmithImporter/Source/DatasmithImporter/Private/DatasmithImportFactory.cpp b/Engine/Plugins/Enterprise/DatasmithImporter/Source/DatasmithImporter/Private/DatasmithImportFactory.cpp
--- a/Engine/Plugins/Enterprise/DatasmithImporter/Source/DatasmithImporter/Private/DatasmithImportFactory.cpp
+++ b/Engine/Plugins/Enterprise/DatasmithImporter/Source/DatasmithImporter/Private/DatasmithImportFactory.cpp
@@ -900,8 +900,11 @@ EReimportResult::Type UDatasmithImportFactory::ReimportMaterial( UMaterialInterf
 
 	FDatasmithImporter::ImportMaterial( ImportContext, MaterialElement.ToSharedRef(), Material );
 
+	const FString& RootFolderPath = ImportContext.AssetsContext.RootFolderPath;
+	const FString& TransientFolderPath = ImportContext.AssetsContext.TransientFolderPath;
+
 	UMaterialInterface* NewMaterial = ImportContext.ImportedMaterials.FindRef( MaterialElement.ToSharedRef() );
-	FDatasmithImporter::FinalizeMaterial( NewMaterial, *MaterialPath, Material );
+	FDatasmithImporter::FinalizeMaterial( NewMaterial, *MaterialPath, *TransientFolderPath, *RootFolderPath, Material );
 
 	FGlobalComponentReregisterContext RecreateComponents;
 
diff --git a/Engine/Plugins/Enterprise/DatasmithImporter/Source/DatasmithImporter/Private/DatasmithImporter.cpp b/Engine/Plugins/Enterprise/DatasmithImporter/Source/DatasmithImporter/Private/DatasmithImporter.cpp
--- a/Engine/Plugins/Enterprise/DatasmithImporter/Source/DatasmithImporter/Private/DatasmithImporter.cpp
+++ b/Engine/Plugins/Enterprise/DatasmithImporter/Source/DatasmithImporter/Private/DatasmithImporter.cpp
@@ -1416,7 +1416,7 @@ UMaterialInterface* FDatasmithImporter::ImportMaterial( FDatasmithImportContext&
 	return ImportedMaterial;
 }
 
-UObject* FDatasmithImporter::FinalizeMaterial( UObject* SourceMaterial, const TCHAR* MaterialsFolderPath, UMaterialInterface* ExistingMaterial, TMap< UObject*, UObject* >* ReferencesToRemap)
+UObject* FDatasmithImporter::FinalizeMaterial( UObject* SourceMaterial, const TCHAR* MaterialFolderPath, const TCHAR* TransientPackagePath, const TCHAR* RootFolderPath, UMaterialInterface* ExistingMaterial, TMap< UObject*, UObject* >* ReferencesToRemap)
 {
 	TRACE_CPUPROFILER_EVENT_SCOPE(FDatasmithImporter::FinalizeMaterial);
 
@@ -1440,15 +1440,13 @@ UObject* FDatasmithImporter::FinalizeMaterial( UObject* SourceMaterial, const TC
 				FString SourceMaterialPath = SourceMaterialInstance->GetOutermost()->GetName();
 				FString SourceParentPath = SourceMaterialParent->GetOutermost()->GetName();
 
-				if (SourceParentPath.StartsWith( SourceMaterialPath ) )
+				if ( SourceParentPath.StartsWith( TransientPackagePath ) )
 				{
 					// Simply finalize the source parent material.
 					// Note that the parent material will be overridden on the existing material instance
-					FString DestinationParentPath = SourceParentPath;
-					DestinationParentPath.RemoveFromStart( SourceMaterialPath );
-					DestinationParentPath = MaterialsFolderPath / DestinationParentPath;
+					const FString DestinationParentPath = SourceParentPath.Replace( TransientPackagePath, RootFolderPath, ESearchCase::CaseSensitive );
 
-					FinalizeMaterial( SourceMaterialParent, *DestinationParentPath, nullptr, ReferencesToRemap );
+					FinalizeMaterial( SourceMaterialParent, *DestinationParentPath, TransientPackagePath, RootFolderPath, nullptr, ReferencesToRemap );
 				}
 			}
 		}
@@ -1456,7 +1454,7 @@ UObject* FDatasmithImporter::FinalizeMaterial( UObject* SourceMaterial, const TC
 
 	UMaterialEditingLibrary::DeleteAllMaterialExpressions( Cast< UMaterial >( ExistingMaterial ) );
 
-	UObject* DestinationMaterial = FinalizeAsset( SourceMaterial, MaterialsFolderPath, ExistingMaterial, ReferencesToRemap );
+	UObject* DestinationMaterial = FinalizeAsset( SourceMaterial, MaterialFolderPath, ExistingMaterial, ReferencesToRemap );
 
 	DatasmithImporterImpl::CompileMaterial( DestinationMaterial );
 
@@ -2446,12 +2444,12 @@ void FDatasmithImporter::FinalizeImport(FDatasmithImportContext& ImportContext,
 			{
 				if (MaterialFunctionInfo.Function && MaterialFunctionInfo.Function->GetOutermost() == SourceMaterial->GetOutermost())
 				{
-					FinalizeMaterial(MaterialFunctionInfo.Function, *DestinationPackagePath, nullptr, &ReferencesToRemap);
+					FinalizeMaterial(MaterialFunctionInfo.Function, *DestinationPackagePath, *TransientFolderPath, *RootFolderPath, nullptr, &ReferencesToRemap);
 				}
 			}
 		}
 
-		ExistingMaterialPtr = FinalizeMaterial(SourceMaterialInterface, *DestinationPackagePath, ExistingMaterial, &ReferencesToRemap);
+		ExistingMaterialPtr = FinalizeMaterial(SourceMaterialInterface, *DestinationPackagePath, *TransientFolderPath, *RootFolderPath, ExistingMaterial, &ReferencesToRemap);
 
 		// Add material to array of packages to apply soft object path redirection to
 		if (ExistingMaterialPtr.IsValid())
diff --git a/Engine/Plugins/Enterprise/DatasmithImporter/Source/DatasmithImporter/Public/DatasmithImporter.h b/Engine/Plugins/Enterprise/DatasmithImporter/Source/DatasmithImporter/Public/DatasmithImporter.h
--- a/Engine/Plugins/Enterprise/DatasmithImporter/Source/DatasmithImporter/Public/DatasmithImporter.h
+++ b/Engine/Plugins/Enterprise/DatasmithImporter/Source/DatasmithImporter/Public/DatasmithImporter.h
@@ -115,13 +115,15 @@ struct DATASMITHIMPORTER_API FDatasmithImporter
 	 * Also finalizes the parent material for instances when it's under the same path as the SourceMaterial.
 	 *
 	 * @param SourceMaterial		The material to finalize.
-	 * @param MaterialsFolderPath	The path where to put the final material, if we don't have an ExistingMaterial
+	 * @param MaterialFolderPath	The path where to put the final material, if we don't have an ExistingMaterial
+	 * @param MaterialFolderPath	The transient path of the imported scene, used to finalize parent materials.
+	 * @param MaterialFolderPath	The root folder where we want to import the scene, used to finalize parent materials.
 	 * @param ExistingMaterial		An optional existing material on which we will finalize over.
 	 * @param ReferencesToRemap		A map of migrated objects. If SourceMaterial refers to any key objects, those references will be replaced by the corresponding map value.
 	 *
 	 * @return The resulting material
 	 */
-	static UObject* FinalizeMaterial( UObject* SourceMaterial, const TCHAR* MaterialsFolderPath, UMaterialInterface* ExistingMaterial, TMap< UObject*, UObject* >* ReferencesToRemap = nullptr );
+	static UObject* FinalizeMaterial( UObject* SourceMaterial, const TCHAR* MaterialFolderPath, const TCHAR* TransientFolderPath, const TCHAR* RootFolderPath, UMaterialInterface* ExistingMaterial, TMap< UObject*, UObject* >* ReferencesToRemap = nullptr );
 
 	/**
 	 * Imports the actors
