Fixed static analysis issues

diff --git a/Engine/Plugins/Enterprise/DatasmithImporter/Source/DatasmithImporter/Private/DatasmithConsumer.cpp b/Engine/Plugins/Enterprise/DatasmithImporter/Source/DatasmithImporter/Private/DatasmithConsumer.cpp
--- a/Engine/Plugins/Enterprise/DatasmithImporter/Source/DatasmithImporter/Private/DatasmithConsumer.cpp
+++ b/Engine/Plugins/Enterprise/DatasmithImporter/Source/DatasmithImporter/Private/DatasmithConsumer.cpp
@@ -593,7 +593,13 @@ bool UDatasmithConsumer::CreateWorld()
 			break;
 		}
 	}
-	ensure(PrimaryLevel);
+
+	// This should never happen
+	if(PrimaryLevel == nullptr)
+	{
+		ensure(false);
+		return false;
+	}
 
 	DatasmithConsumerUtils::SetMarker(PrimaryLevel, ConsumerMarkerID, UniqueID);
 	PrimaryLevel->bIsVisible = true;
@@ -946,10 +952,16 @@ ULevel* UDatasmithConsumer::FindOrAddLevel(const FString& InLevelName)
 	{
 		FTransform LevelTransform;
 		StreamingLevel = UEditorLevelUtils::AddLevelToWorld(WorkingWorld.Get(), *LevelPackageName, ULevelStreamingAlwaysLoaded::StaticClass(), LevelTransform);
-		ensure(StreamingLevel);
+		if(StreamingLevel)
+		{
 
-		WorkingWorld->LoadSecondaryLevels();
-		ensure(StreamingLevel->GetLoadedLevel());
+			WorkingWorld->LoadSecondaryLevels();
+			ensure(StreamingLevel->GetLoadedLevel());
+		}
+		else
+		{
+			ensure(false);
+		}
 
 		bCleanLevel = true;
 	}
@@ -965,35 +977,41 @@ ULevel* UDatasmithConsumer::FindOrAddLevel(const FString& InLevelName)
 	// Mark level as generated by this consumer
 	if(StreamingLevel)
 	{
-		ULevel* NewLevel = StreamingLevel->GetLoadedLevel();
-
-		DatasmithConsumerUtils::SetMarker(NewLevel, ConsumerMarkerID, UniqueID);
-
-		if(bCleanLevel)
+		if(ULevel* NewLevel = StreamingLevel->GetLoadedLevel())
 		{
-			// Clean up the level if it contains actors from previous execution
-			UWorld* LevelWorld = Cast<UWorld>(NewLevel->GetOuter());
+			DatasmithConsumerUtils::SetMarker(NewLevel, ConsumerMarkerID, UniqueID);
 
-			TArray<AActor*> LevelActors = NewLevel->Actors;
-			for(AActor* Actor : LevelActors)
+			if(bCleanLevel)
 			{
+				// Clean up the level if it contains actors from previous execution
+				UWorld* LevelWorld = Cast<UWorld>(NewLevel->GetOuter());
 
-				if(Actor && Actor->GetRootComponent() && !Actor->IsA<AWorldSettings>() && !Actor->IsA<APhysicsVolume>() && !Actor->IsA<ABrush>() )
+				TArray<AActor*> LevelActors = NewLevel->Actors;
+				for(AActor* Actor : LevelActors)
 				{
-					LevelWorld->DestroyActor(Actor, true);
-					Actor->UObject::Rename(nullptr, GetTransientPackage(), REN_DontCreateRedirectors | REN_ForceNoResetLoaders);
+
+					if(Actor && Actor->GetRootComponent() && !Actor->IsA<AWorldSettings>() && !Actor->IsA<APhysicsVolume>() && !Actor->IsA<ABrush>() )
+					{
+						LevelWorld->DestroyActor(Actor, true);
+						Actor->UObject::Rename(nullptr, GetTransientPackage(), REN_DontCreateRedirectors | REN_ForceNoResetLoaders);
+					}
 				}
+
+				LevelWorld->CleanupWorld(false, false);
+				LevelWorld->CleanupActors();
+
+				DatasmithConsumerUtils::SaveMap(LevelWorld);
 			}
 
-			LevelWorld->CleanupWorld(false, false);
-			LevelWorld->CleanupActors();
+			WorkingWorld->AddLevel(NewLevel);
 
-			DatasmithConsumerUtils::SaveMap(LevelWorld);
+			return NewLevel;
+		}
+		else
+		{
+			ensure(false);
 		}
 
-		WorkingWorld->AddLevel(NewLevel);
-
-		return NewLevel;
 	}
 
 	return nullptr;
diff --git a/Engine/Plugins/Enterprise/DatasmithImporter/Source/DatasmithImporter/Private/UI/DatasmithConsumerDetails.cpp b/Engine/Plugins/Enterprise/DatasmithImporter/Source/DatasmithImporter/Private/UI/DatasmithConsumerDetails.cpp
--- a/Engine/Plugins/Enterprise/DatasmithImporter/Source/DatasmithImporter/Private/UI/DatasmithConsumerDetails.cpp
+++ b/Engine/Plugins/Enterprise/DatasmithImporter/Source/DatasmithImporter/Private/UI/DatasmithConsumerDetails.cpp
@@ -269,116 +269,6 @@ namespace DatasmithConsumerDetailsUtil
 		TWeakObjectPtr< UDatasmithConsumer > ConsumerPtr;
 		TSharedPtr< SEditableTextBox > LevelTextBox;
 	};
-
-	class SLevelAssetProperty : public SCompoundWidget
-	{
-	public:
-		SLATE_BEGIN_ARGS(SLevelAssetProperty) {}
-		SLATE_END_ARGS()
-
-	public:
-		void Construct(const FArguments& InArgs, UDatasmithConsumer* InConsumer)
-		{
-			ConsumerPtr = InConsumer;
-			FSlateFontInfo FontInfo = IDetailLayoutBuilder::GetDetailFont();
-
-			TArray<FAssetData> AssetDataArray;
-			AssetDataArray.Add(FAssetData(ConsumerPtr->GetOutermost()));
-
-			TSharedPtr<SWidget> OutputLevelWidget = SNew(SObjectPropertyEntryBox)
-				.AllowedClass(UWorld::StaticClass())
-				.ObjectPath(this, &SLevelAssetProperty::GetAssetPath)
-				.OnShouldFilterAsset(this, &SLevelAssetProperty::OnShouldFilterAsset)
-				.OnObjectChanged(this, &SLevelAssetProperty::OnAssetSelected)
-				.OwnerAssetDataArray(AssetDataArray)
-				.AllowClear( false )
-				.DisplayUseSelected( true )
-				.DisplayBrowse( true )
-				.EnableContentPicker(true)
-				.DisplayCompactSize(true)
-				.DisplayThumbnail(false);
-
-			ChildSlot
-			[
-				SNew(SConstrainedBox)
-				[
-					OutputLevelWidget.ToSharedRef()
-				]
-			];
-		}
-
-	private:
-		void OnAssetSelected(const FAssetData& AssetData)
-		{
-			if ( UDatasmithConsumer* Consumer = ConsumerPtr.Get() )
-			{
-				UWorld* World = Cast<UWorld>(AssetData.GetAsset());
-				ensure(World);
-
-				// Check that level has been chosen or created in the right content folder
-				if(ULevel* Level = World->PersistentLevel)
-				{
-					const FString PackagePathName = FPaths::GetPath(Level->GetOutermost()->GetPathName());
-					const FString TargetContentFolder = Consumer->GetTargetContentFolder();
-					if(PackagePathName == TargetContentFolder)
-					{
-						FScopedTransaction Transaction( LOCTEXT("Consumer_SetOutputLevel", "Set Output Level") );
-						FText OutReason;
-
-						if(Consumer->SetLevelName(AssetData.AssetName.ToString(), OutReason))
-						{
-							WorldPtr = World;
-						}
-						else
-						{
-							Transaction.Cancel();
-							UE_LOG( LogDatasmithImport, Error, TEXT("Cannot create a level named %s - %s"), *AssetData.AssetName.ToString(), *OutReason.ToString() );
-						}
-					}
-					// Inform user about where the level should be created
-					else
-					{
-						const FText WarningMessage = FText::Format(LOCTEXT("LevelProperty_WrongFolder", "The new level must be created in the asset's folder, {0}."), FText::FromString(TargetContentFolder));
-						const FText DialogTitle(LOCTEXT("LevelProperty_WrongFolder_Title", "Warning - Wrong Folder"));
-
-						FMessageDialog::Open(EAppMsgType::Ok, WarningMessage, &DialogTitle);
-					}
-				}
-			}
-		}
-
-		FString GetAssetPath() const
-		{
-			//if ( UDatasmithConsumer* Consumer = ConsumerPtr.Get() )
-			//{
-			//	if(!WorldPtr.IsValid())
-			//	{
-			//		FString OutputLevelPath = Consumer->GetOutputLevelPath();
-			//	}
-
-			//	return FSoftObjectPath(WorldPtr->PersistentLevel).GetAssetPathString();
-			//}
-
-			return FString("");
-		}
-
-		bool OnShouldFilterAsset(const FAssetData& AssetData) const
-		{
-			if ( UDatasmithConsumer* Consumer = ConsumerPtr.Get() )
-			{
-				bool bValidAsset = AssetData.PackagePath.ToString() == Consumer->GetTargetContentFolder();
-				bValidAsset &= AssetData.AssetClass == UWorld::StaticClass()->GetFName();
-
-				return !bValidAsset;
-			}
-
-			return true;
-		}
-
-	private:
-		TWeakObjectPtr< UDatasmithConsumer > ConsumerPtr;
-		TWeakObjectPtr< UWorld > WorldPtr;
-	};
 }
 
 void FDatasmithConsumerDetails::CustomizeDetails(IDetailLayoutBuilder& DetailBuilder)
