Fixed issue with remaining empty dirty packages after the commit of a dataprep asset containing output folder operations
Revisited the resolution of soft references in assets after moves of assets and/or actors
Reworked the move of assets from one folder to another

diff --git a/Engine/Plugins/Enterprise/DatasmithImporter/Source/DatasmithImporter/Private/DatasmithConsumer.cpp b/Engine/Plugins/Enterprise/DatasmithImporter/Source/DatasmithImporter/Private/DatasmithConsumer.cpp
--- a/Engine/Plugins/Enterprise/DatasmithImporter/Source/DatasmithImporter/Private/DatasmithConsumer.cpp
+++ b/Engine/Plugins/Enterprise/DatasmithImporter/Source/DatasmithImporter/Private/DatasmithConsumer.cpp
@@ -98,7 +98,7 @@ namespace DatasmithConsumerUtils
 
 	void SetMarker(UObject* Object, const FString& Name, const FString& Value);
 
-	void MoveActorsToLevel(const TArray<AActor*>& ActorsToMove, ULevel* DestLevel, TMap< FName, TSoftObjectPtr< AActor > >& ActorsMap, bool bDuplicate);
+	void MoveActorsToLevel(const TArray<AActor*>& ActorsToMove, ULevel* DestLevel, TMap< FName, TSoftObjectPtr< AActor > >& ActorsMap, const TArray<UPackage*>& PackagesToCheck, bool bDuplicate);
 
 	template<class AssetClass>
 	void SetMarker(const TMap<FName, TSoftObjectPtr< AssetClass >>& AssetMap, const FString& Name, const FString& Value)
@@ -130,7 +130,7 @@ namespace DatasmithConsumerUtils
 	}
 
 	template<class AssetClass>
-	void ApplyFolderDirective(TMap<FName, TSoftObjectPtr< AssetClass >>& AssetMap, const FString& RootPackagePath, TFunction<void(ELogVerbosity::Type, FText)> ReportCallback)
+	TArray<UPackage*> ApplyFolderDirective(TMap<FName, TSoftObjectPtr< AssetClass >>& AssetMap, const FString& RootPackagePath, TMap<FSoftObjectPath, FSoftObjectPath>& AssetRedirectorMap, TFunction<void(ELogVerbosity::Type, FText)> ReportCallback)
 	{
 		auto CanMoveAsset = [&ReportCallback](UObject* Source, UObject* Target) -> bool
 		{
@@ -163,8 +163,11 @@ namespace DatasmithConsumerUtils
 			return false;
 		};
 
-		TArray<UPackage*> PackagesToCheck;
-		TMap<FSoftObjectPath, FSoftObjectPath> AssetRedirectorMap;
+		TArray<UPackage*> PackagesProcessed;
+		PackagesProcessed.Reserve(AssetMap.Num());
+
+		// Reserve room for new entries in remapping table
+		AssetRedirectorMap.Reserve(AssetRedirectorMap.Num() + AssetMap.Num());
 
 		for(TPair<FName, TSoftObjectPtr< AssetClass >>& Entry : AssetMap)
 		{
@@ -173,14 +176,15 @@ namespace DatasmithConsumerUtils
 				const FString& OutputFolder = GetMarker(Asset, UDataprepContentConsumer::RelativeOutput);
 				if(OutputFolder.Len() > 0)
 				{
-					FString SourcePackagePath = Entry.Value->GetOuter()->GetPathName();
+					UPackage* SourcePackage = Entry.Value->GetOutermost();
 					FString TargetPackagePath = FPaths::Combine(RootPackagePath, OutputFolder, Asset->GetName());
-					
-					FString PackageFilename;
-					FPackageName::TryConvertLongPackageNameToFilename( TargetPackagePath, PackageFilename, FPackageName::GetAssetPackageExtension() );
 
-					if( SourcePackagePath != TargetPackagePath)
+
+					if( ensure(SourcePackage) && SourcePackage->GetPathName() != TargetPackagePath)
 					{
+						FString PackageFilename;
+						FPackageName::TryConvertLongPackageNameToFilename( TargetPackagePath, PackageFilename, FPackageName::GetAssetPackageExtension() );
+
 						bool bCanMove = true;
 
 						FString TargetAssetFullPath = TargetPackagePath + "." + Asset->GetName();
@@ -197,25 +201,37 @@ namespace DatasmithConsumerUtils
 						{
 							FSoftObjectPath& SoftObjectPathRef = AssetRedirectorMap.Emplace( Asset );
 
-							UPackage* Package = CreatePackage(nullptr, *TargetPackagePath);
-							Package->FullyLoad();
+							UPackage* TargetPackage = CreatePackage(nullptr, *TargetPackagePath);
+							TargetPackage->FullyLoad();
 
-							Asset->Rename(nullptr, Package, REN_DontCreateRedirectors | REN_NonTransactional);
-							Entry.Value = Asset;
+							Asset->Rename(nullptr, TargetPackage, REN_DontCreateRedirectors | REN_NonTransactional);
 
+							// Update asset registry with renaming
+							FAssetRegistryModule::AssetRenamed( Asset, SourcePackage->GetPathName() + TEXT(".") + Asset->GetName() );
+
+							Entry.Value = Asset;
 							SoftObjectPathRef = Asset;
-							PackagesToCheck.Add(Package);
+							PackagesProcessed.Add(TargetPackage);
+
+							// Clean up flags on source package. It is not useful anymore
+							SourcePackage->SetDirtyFlag(false);
+							SourcePackage->SetFlags(RF_Transient);
+							SourcePackage->ClearFlags(RF_Standalone | RF_Public);
 						}
 					}
+					else
+					{
+						PackagesProcessed.Add(SourcePackage);
+					}
+				}
+				else
+				{
+					PackagesProcessed.Add(Entry.Value->GetOutermost());
 				}
 			}
 		}
 
-		if(AssetRedirectorMap.Num() > 0)
-		{
-			IAssetTools& AssetTools = FModuleManager::LoadModuleChecked<FAssetToolsModule>("AssetTools").Get();
-			AssetTools.RenameReferencingSoftObjectPaths(PackagesToCheck, AssetRedirectorMap);
-		}
+		return PackagesProcessed;
 	}
 }
 
@@ -459,26 +475,41 @@ bool UDatasmithConsumer::Run()
 		}
 	};
 
+	// Array to store materials which soft references might need to be fixed
+	TArray<UPackage*> PackagesToCheck;
+
+	// Array to store level sequences and level variants which soft references might need to be fixed
+	TArray<UPackage*> PackagesToFix;
+
+	// Table of remapping to contain moved assets
+	TMap<FSoftObjectPath, FSoftObjectPath> AssetRedirectorMap;
+
 	DatasmithConsumerUtils::SetMarker(SceneAsset->Textures, UDatasmithConsumer::ConsumerMarkerID, UniqueID);
-	DatasmithConsumerUtils::ApplyFolderDirective(SceneAsset->Textures, TargetContentFolder, ReportFunc );
+	DatasmithConsumerUtils::ApplyFolderDirective(SceneAsset->Textures, TargetContentFolder, AssetRedirectorMap, ReportFunc );
 
 	DatasmithConsumerUtils::SetMarker(SceneAsset->StaticMeshes, UDatasmithConsumer::ConsumerMarkerID, UniqueID);
-	DatasmithConsumerUtils::ApplyFolderDirective(SceneAsset->StaticMeshes, TargetContentFolder, ReportFunc );
+	DatasmithConsumerUtils::ApplyFolderDirective(SceneAsset->StaticMeshes, TargetContentFolder, AssetRedirectorMap, ReportFunc );
 
 	DatasmithConsumerUtils::SetMarker(SceneAsset->Materials, UDatasmithConsumer::ConsumerMarkerID, UniqueID);
-	DatasmithConsumerUtils::ApplyFolderDirective(SceneAsset->Materials, TargetContentFolder, ReportFunc );
+	PackagesToCheck.Append(DatasmithConsumerUtils::ApplyFolderDirective(SceneAsset->Materials, TargetContentFolder, AssetRedirectorMap, ReportFunc ));
 
 	DatasmithConsumerUtils::SetMarker(SceneAsset->MaterialFunctions, UDatasmithConsumer::ConsumerMarkerID, UniqueID);
-	DatasmithConsumerUtils::ApplyFolderDirective(SceneAsset->MaterialFunctions, TargetContentFolder, ReportFunc );
+	DatasmithConsumerUtils::ApplyFolderDirective(SceneAsset->MaterialFunctions, TargetContentFolder, AssetRedirectorMap, ReportFunc );
 
 	DatasmithConsumerUtils::SetMarker(SceneAsset->LevelSequences, UDatasmithConsumer::ConsumerMarkerID, UniqueID);
-	DatasmithConsumerUtils::ApplyFolderDirective(SceneAsset->LevelSequences, TargetContentFolder, ReportFunc );
+	PackagesToFix.Append(DatasmithConsumerUtils::ApplyFolderDirective(SceneAsset->LevelSequences, TargetContentFolder, AssetRedirectorMap, ReportFunc ));
 
 	DatasmithConsumerUtils::SetMarker(SceneAsset->LevelVariantSets, UDatasmithConsumer::ConsumerMarkerID, UniqueID);
-	DatasmithConsumerUtils::ApplyFolderDirective(SceneAsset->LevelVariantSets, TargetContentFolder, ReportFunc );
+	PackagesToFix.Append(DatasmithConsumerUtils::ApplyFolderDirective(SceneAsset->LevelVariantSets, TargetContentFolder, AssetRedirectorMap, ReportFunc ));
+
+	if(AssetRedirectorMap.Num() > 0 && PackagesToCheck.Num() > 0)
+	{
+		IAssetTools& AssetTools = FModuleManager::LoadModuleChecked<FAssetToolsModule>("AssetTools").Get();
+		AssetTools.RenameReferencingSoftObjectPaths(PackagesToCheck, AssetRedirectorMap);
+	}
 
 	// Apply UDataprepConsumerUserData directives for actors
-	ApplySubLevelDirective();
+	ApplySubLevelDirective(PackagesToFix);
 
 	return FinalizeRun();
 }
@@ -629,7 +660,7 @@ bool UDatasmithConsumer::CreateWorld()
 					}
 				}
 
-				DatasmithConsumerUtils::MoveActorsToLevel( ActorsToCopy, PrimaryLevel, RelatedActors, true);
+				DatasmithConsumerUtils::MoveActorsToLevel( ActorsToCopy, PrimaryLevel, RelatedActors, TArray<UPackage*>(), true);
 
 				World->RemoveLevel(Level);
 
@@ -1173,7 +1204,7 @@ bool UDatasmithConsumer::CheckOutputDirectives()
 	return true;
 }
 
-void UDatasmithConsumer::ApplySubLevelDirective()
+void UDatasmithConsumer::ApplySubLevelDirective(const TArray<UPackage*>& PackagesToCheck)
 {
 	TMap< FName, TSoftObjectPtr< AActor > >& RelatedActors = ImportContextPtr->ActorsContext.CurrentTargetedScene->RelatedActors;
 
@@ -1228,9 +1259,11 @@ void UDatasmithConsumer::ApplySubLevelDirective()
 		}
 	}
 
+	TMap<FSoftObjectPath, FSoftObjectPath> AssetRedirectorMap;
+
 	for(TPair<ULevel*, TArray<AActor*>> Entry : ActorsToMove)
 	{
-		DatasmithConsumerUtils::MoveActorsToLevel( Entry.Value, Entry.Key, RelatedActors, false);
+		DatasmithConsumerUtils::MoveActorsToLevel( Entry.Value, Entry.Key, RelatedActors, PackagesToCheck, false);
 	}
 }
 
@@ -1629,7 +1662,7 @@ namespace DatasmithConsumerUtils
 		WorldToSave->GetOutermost()->SetDirtyFlag(false);
 	}
 
-	TArray<AActor*> MoveActorsToLevel(const TArray<AActor*>& ActorsToMove, ULevel* DestLevel, bool bDuplicate)
+	TArray<AActor*> MoveActorsToLevel(const TArray<AActor*>& ActorsToMove, ULevel* DestLevel, const TArray<UPackage*>& PackagesToCheck, bool bDuplicate)
 	{
 		if(DestLevel == nullptr || ActorsToMove.Num() == 0)
 		{
@@ -1742,7 +1775,7 @@ namespace DatasmithConsumerUtils
 			}
 		}
 
-		FAssetToolsModule& AssetToolsModule = FModuleManager::GetModuleChecked<FAssetToolsModule>(TEXT("AssetTools"));
+		IAssetTools& AssetTools = FModuleManager::LoadModuleChecked<FAssetToolsModule>("AssetTools").Get();
 		TArray<FAssetRenameData> RenameData;
 
 		for (TPair<FSoftObjectPath, FSoftObjectPath>& Pair : ActorPathMapping)
@@ -1755,7 +1788,13 @@ namespace DatasmithConsumerUtils
 
 		if (RenameData.Num() > 0)
 		{
-			AssetToolsModule.Get().RenameAssets(RenameData);
+			AssetTools.RenameAssets(RenameData);
+
+			// Fix soft references in level sequences and variants
+			if(PackagesToCheck.Num() > 0)
+			{
+				AssetTools.RenameReferencingSoftObjectPaths(PackagesToCheck, ActorPathMapping);
+			}
 		}
 
 		// Restore the original clipboard contents
@@ -1764,7 +1803,7 @@ namespace DatasmithConsumerUtils
 		return NewActors;
 	}
 
-	void MoveActorsToLevel(const TArray<AActor*>& ActorsToMove, ULevel* DestLevel, TMap<FName,TSoftObjectPtr<AActor>>& ActorsMap, bool bDuplicate)
+	void MoveActorsToLevel(const TArray<AActor*>& ActorsToMove, ULevel* DestLevel, TMap<FName,TSoftObjectPtr<AActor>>& ActorsMap, const TArray<UPackage*>& PackagesToCheck, bool bDuplicate)
 	{
 		if(ActorsToMove.Num() > 0)
 		{
@@ -1776,8 +1815,7 @@ namespace DatasmithConsumerUtils
 			EObjectFlags DestWorldFlags = DestLevel->GetOuter()->GetFlags();
 			EObjectFlags DestPackageFlags = DestLevel->GetOutermost()->GetFlags();
 
-			TArray<AActor*> NewActors = MoveActorsToLevel( ActorsToMove, DestLevel, bDuplicate);
-			printf(">>> %d", NewActors.Num());
+			TArray<AActor*> NewActors = MoveActorsToLevel( ActorsToMove, DestLevel, PackagesToCheck, bDuplicate);
 
 			GWorld = PrevGWorld;
 
diff --git a/Engine/Plugins/Enterprise/DatasmithImporter/Source/DatasmithImporter/Private/DatasmithConsumer.h b/Engine/Plugins/Enterprise/DatasmithImporter/Source/DatasmithImporter/Private/DatasmithConsumer.h
--- a/Engine/Plugins/Enterprise/DatasmithImporter/Source/DatasmithImporter/Private/DatasmithConsumer.h
+++ b/Engine/Plugins/Enterprise/DatasmithImporter/Source/DatasmithImporter/Private/DatasmithConsumer.h
@@ -20,6 +20,7 @@ namespace DatasmithConsumerDetailsUtil
 }
 class UDatasmithScene;
 class ULevel;
+class UPackage;
 class UWorld;
 
 UCLASS(Experimental, config = EditorSettings, HideCategories = (DatasmithConsumerInternal))
@@ -71,7 +72,7 @@ private:
 	void UpdateScene();
 
 	/** Process actors which have been marked with a UDataprepConsumerUserData */
-	void ApplySubLevelDirective();
+	void ApplySubLevelDirective(const TArray<UPackage*>& PackagesToCheck);
 
 	/** Find or add a new level to the target world */
 	ULevel* FindOrAddLevel(const FString& InLevelName);
