Reworked the creation of levels in the DatasmithConsumer following decision made in meeting of 2/27/20:
 - Only one level is created and added to the world edited in the LevelEditor
Removed "Output to Level" operation
Added initialization of boolean in UDataprepActionStep

#rnx

diff --git a/Engine/Plugins/Enterprise/DataprepEditor/Source/DataprepCore/Public/DataprepActionAsset.h b/Engine/Plugins/Enterprise/DataprepEditor/Source/DataprepCore/Public/DataprepActionAsset.h
--- a/Engine/Plugins/Enterprise/DataprepEditor/Source/DataprepCore/Public/DataprepActionAsset.h
+++ b/Engine/Plugins/Enterprise/DataprepEditor/Source/DataprepCore/Public/DataprepActionAsset.h
@@ -54,6 +54,11 @@ class UDataprepActionStep : public UObject
 
 public:
 
+	UDataprepActionStep()
+		: bIsEnabled(true)
+		, StepObject(nullptr)
+	{}
+
 	// Begin UObject Interface
 	virtual void PostLoad() override;
 	// End UObject Interface
diff --git a/Engine/Plugins/Enterprise/DataprepEditor/Source/DataprepLibraries/Private/DataprepOperations.cpp b/Engine/Plugins/Enterprise/DataprepEditor/Source/DataprepLibraries/Private/DataprepOperations.cpp
--- a/Engine/Plugins/Enterprise/DataprepEditor/Source/DataprepLibraries/Private/DataprepOperations.cpp
+++ b/Engine/Plugins/Enterprise/DataprepEditor/Source/DataprepLibraries/Private/DataprepOperations.cpp
@@ -383,15 +383,6 @@ void UDataprepFlipFacesOperation::OnExecution_Implementation(const FDataprepCont
 	UStaticMesh::BatchBuild(StaticMeshes.Array());
 }
 
-void UDataprepSetOutputLevel::OnExecution_Implementation(const FDataprepContext& InContext)
-{
-#ifdef LOG_TIME
-	DataprepOperationTime::FTimeLogger TimeLogger(TEXT("RandomizeTransform"), [&](FText Text) { this->LogInfo(Text); });
-#endif
-
-	UDataprepOperationsLibrary::SetSubOuputLevel(InContext.Objects, LevelName);
-}
-
 void UDataprepSetOutputFolder::OnExecution_Implementation(const FDataprepContext& InContext)
 {
 #ifdef LOG_TIME
diff --git a/Engine/Plugins/Enterprise/DataprepEditor/Source/DataprepLibraries/Private/DataprepOperations.h b/Engine/Plugins/Enterprise/DataprepEditor/Source/DataprepLibraries/Private/DataprepOperations.h
--- a/Engine/Plugins/Enterprise/DataprepEditor/Source/DataprepLibraries/Private/DataprepOperations.h
+++ b/Engine/Plugins/Enterprise/DataprepEditor/Source/DataprepLibraries/Private/DataprepOperations.h
@@ -449,33 +449,6 @@ protected:
 	//~ End UDataprepOperation Interface
 };
 
-UCLASS(Experimental, Category = ActorOperation, Meta = (DisplayName="Output to Level", ToolTip = "For each actor to process, set the sub-level to move it to.\nThe sub-level is relative to the level specified to the Dataprep consumer.") )
-class UDataprepSetOutputLevel : public UDataprepOperation
-{
-	GENERATED_BODY()
-
-	UDataprepSetOutputLevel()
-	: LevelName(TEXT("MySubLevel"))
-	{
-	}
-
-public:
-	// If true, the screen sizes at which LODs swap are computed automatically
-	UPROPERTY(EditAnywhere, BlueprintReadWrite, Category = ActorOperation, meta = (DisplayName = "Level Name", ToolTip = "Name of the sub level the actors to be saved to"))
-	FString LevelName;
-
-	//~ Begin UDataprepOperation Interface
-public:
-	virtual FText GetCategory_Implementation() const override
-	{
-		return FDataprepOperationCategories::ActorOperation;
-	}
-
-protected:
-	virtual void OnExecution_Implementation(const FDataprepContext& InContext) override;
-	//~ End UDataprepOperation Interface
-};
-
 UCLASS(Experimental, Category = AssetOperation, Meta = (DisplayName="Output to Folder", ToolTip = "For each asset to process, set the sub-folder to save it to.\nThe sub-folder is relative to the folder specified to the Dataprep consumer.") )
 class UDataprepSetOutputFolder : public UDataprepOperation
 {
diff --git a/Engine/Plugins/Enterprise/DatasmithImporter/Source/DatasmithImporter/Private/DatasmithConsumer.cpp b/Engine/Plugins/Enterprise/DatasmithImporter/Source/DatasmithImporter/Private/DatasmithConsumer.cpp
--- a/Engine/Plugins/Enterprise/DatasmithImporter/Source/DatasmithImporter/Private/DatasmithConsumer.cpp
+++ b/Engine/Plugins/Enterprise/DatasmithImporter/Source/DatasmithImporter/Private/DatasmithConsumer.cpp
@@ -28,6 +28,7 @@
 #include "ObjectTemplates/DatasmithSkyLightComponentTemplate.h"
 #include "ObjectTemplates/DatasmithStaticMeshComponentTemplate.h"
 #include "Utility/DatasmithImporterUtils.h"
+#include "Utility/DatasmithImporterImpl.h"
 
 #include "Algo/Count.h"
 #include "AssetRegistryModule.h"
@@ -407,32 +408,9 @@ bool UDatasmithConsumer::Initialize()
 	}
 
 	// Check if the finalize should be threated as a reimport
-	TArray< ADatasmithSceneActor* > SceneActors = FDatasmithImporterUtils::FindSceneActors( ImportContextPtr->ActorsContext.FinalWorld, ImportContextPtr->SceneAsset);
-	if (SceneActors.Num() > 0 )
+	ImportContextPtr->ActorsContext.FinalSceneActors.Append(FDatasmithImporterUtils::FindSceneActors( ImportContextPtr->ActorsContext.FinalWorld, ImportContextPtr->SceneAsset));
+	if (ImportContextPtr->ActorsContext.FinalSceneActors.Num() > 0 )
 	{
-		ADatasmithSceneActor* FoundSceneActor = nullptr;
-		for( ADatasmithSceneActor* SceneActor : SceneActors )
-		{
-			if( SceneActor && SceneActor->Scene == DatasmithScene )
-			{
-				FoundSceneActor = SceneActor;
-				break;
-			}
-		}
-
-		if( FoundSceneActor == nullptr )
-		{
-			//Create a new Datasmith scene actor in the targeted level
-			FActorSpawnParameters SpawnParameters;
-			SpawnParameters.Template = ImportContextPtr->ActorsContext.ImportSceneActor;
-			ADatasmithSceneActor* DestinationSceneActor = Cast< ADatasmithSceneActor >(ImportContextPtr->ActorsContext.FinalWorld->SpawnActor< ADatasmithSceneActor >(SpawnParameters));
-
-			// Name new destination ADatasmithSceneActor to the DatasmithScene's name
-			DestinationSceneActor->SetActorLabel( ImportContextPtr->Scene->GetName() );
-			DestinationSceneActor->MarkPackageDirty();
-			DestinationSceneActor->RelatedActors.Reset();
-		}
-
 		ImportContextPtr->bIsAReimport = true;
 		ImportContextPtr->Options->ReimportOptions.bRespawnDeletedActors = false;
 		ImportContextPtr->Options->ReimportOptions.bUpdateActors = true;
@@ -508,10 +486,7 @@ bool UDatasmithConsumer::Run()
 		AssetTools.RenameReferencingSoftObjectPaths(PackagesToCheck, AssetRedirectorMap);
 	}
 
-	// Apply UDataprepConsumerUserData directives for actors
-	ApplySubLevelDirective(PackagesToFix);
-
-	return FinalizeRun();
+	return true;
 }
 
 bool UDatasmithConsumer::FinalizeRun()
@@ -556,63 +531,13 @@ bool UDatasmithConsumer::FinalizeRun()
 	return true;
 }
 
-
 bool UDatasmithConsumer::CreateWorld()
 {
 	ensure(!WorkingWorld.IsValid());
 
-	UWorld* World = Cast<UWorld>(OutputLevelSoftObject.TryLoad());
-	if(World)
-	{
-		World->SetFlags(RF_Public | RF_Transactional |  RF_Standalone);
-
-		FWorldContext& WorldContext = GEngine->CreateNewWorldContext(World->WorldType);
-		WorldContext.SetCurrentWorld(World);
+	WorkingWorld = TStrongObjectPtr<UWorld>(GWorld.GetReference());
 
-		// Load all the secondary levels of the world.
-		World->LoadSecondaryLevels(true);
-
-		// Check that all secondary levels have been added to world
-		for(ULevelStreaming* StreamingLevel : World->GetStreamingLevels())
-		{
-			if(ULevel* SecondaryLevel = StreamingLevel->GetLoadedLevel())
-			{
-				if(!World->ContainsLevel(SecondaryLevel))
-				{
-					World->AddLevel(SecondaryLevel);
-				}
-			}
-		}
-	}
-	else
-	{
-		UPackage* Package = CreatePackage(nullptr, *OutputLevelSoftObject.GetLongPackageName());
-		Package->FullyLoad();
-		Package->SetFlags(RF_Public);
-
-		World = NewObject< UWorld >(Package, *OutputLevelSoftObject.GetAssetName(), RF_Public | RF_Transactional |  RF_Standalone);
-		World->WorldType = EWorldType::Inactive;
-
-		FWorldContext& WorldContext = GEngine->CreateNewWorldContext(World->WorldType);
-		WorldContext.SetCurrentWorld(World);
-
-		World->InitializeNewWorld(UWorld::InitializationValues()
-			.AllowAudioPlayback(false)
-			.CreateAISystem(false)
-			.CreateNavigation(false)
-			.CreatePhysicsScene(false)
-			.RequiresHitProxies(false)
-			.ShouldSimulatePhysics(false)
-			.SetTransactional(false));
-	}
-
-	if(World == nullptr)
-	{
-		ensure(false);
-		return false;
-	}
-
-	TArray<ULevel*> Levels = World->GetLevels();
+	TArray<ULevel*> Levels = WorkingWorld->GetLevels();
 
 	// Find level associated with this consumer
 	PrimaryLevel = nullptr;
@@ -620,84 +545,73 @@ bool UDatasmithConsumer::CreateWorld()
 	{
 		if(Level && Level->GetOuter()->GetName() == LevelName)
 		{
-			PrimaryLevel = Level;
-			break;
-		}
-	}
-
-	// This should never happen
-	if(PrimaryLevel == nullptr)
-	{
-		ensure(false);
-		return false;
-	}
-
-	DatasmithConsumerUtils::SetMarker(PrimaryLevel, ConsumerMarkerID, UniqueID);
-	PrimaryLevel->bIsVisible = true;
-
-	// If there is more than one level, move all actors of the world to the main level.
-	// The call to FinalizeRun will take care of redistributing to the sub-levels if applicable
-	if(Levels.Num() > 1)
-	{
-		// Get the ADatasmithSceneActor of the world if it exists
-		TArray< ADatasmithSceneActor* > SceneActors = FDatasmithImporterUtils::FindSceneActors( World, DatasmithScene.Get());
-		TMap< FName, TSoftObjectPtr< AActor > > EmptyRelatedActors;
-		TMap< FName, TSoftObjectPtr< AActor > >& RelatedActors = SceneActors.Num() > 0 ? SceneActors[0]->RelatedActors : EmptyRelatedActors;
-
-		for(ULevel* Level : Levels)
-		{
-			if(Level != PrimaryLevel)
+			// Found a level with the same name, make sure it is the same package path
+			if(UWorld* LevelWorld = Cast<UWorld>(Level->GetOuter()))
 			{
-				Level->bIsVisible = true;
+				FSoftObjectPath LevelWorldPath(LevelWorld);
 
-				// Collect actors to copy to primary level
-				TArray<AActor*> ActorsToCopy;
-				for(AActor* Actor : Level->Actors)
+				// If paths differ, remove level with same name from world
+				if(LevelWorldPath != OutputLevelSoftObject)
 				{
-					if(Actor && Actor->GetRootComponent() && !Actor->IsA<AWorldSettings>() && !Actor->IsA<APhysicsVolume>() && !Actor->IsA<ABrush>() )
+					if(DatasmithConsumerUtils::GetMarker(Level, ConsumerMarkerID) == UniqueID)
 					{
-						ActorsToCopy.Add(Actor);
-					}
-				}
-
-				DatasmithConsumerUtils::MoveActorsToLevel( ActorsToCopy, PrimaryLevel, RelatedActors, TArray<UPackage*>(), true);
-
-				World->RemoveLevel(Level);
-
-				// Make sure the world and package are properly discarded
-				UPackage* LevelPackage = Level->GetOutermost();
-				UWorld* LevelWorld = Cast<UWorld>(Level->GetOuter());
-
-				// Move the world to the transient package
-				LevelWorld->Rename(nullptr, GetTransientPackage(), REN_NonTransactional | REN_DontCreateRedirectors);
-
-				// Empty the world from all its content
-				LevelWorld->DestroyWorld( true );
-				GEngine->DestroyWorldContext( LevelWorld );
+						if(Levels.Num() > 0)
+						{
+							const bool bShowDialog = !Context.bSilentMode && !IsRunningCommandlet();
+							bool bUnloadLevel = true;
 
-				// Indicates world is good for garbage collect
-				LevelWorld->ClearFlags(RF_Standalone | RF_Public | RF_Transactional);
-				LevelWorld->SetFlags(RF_Transient);
-				LevelWorld->MarkPendingKill();
+							if(bShowDialog)
+							{
+								const FTextFormat Format(LOCTEXT("DatasmithConsumer_AlreadyLoaded_Dlg", "Level {0} with a different path from the same Dataprep asset is already loaded.\n\nDo you want to unload it?"));
+								const FText WarningMessage = FText::Format( Format, FText::FromString(LevelName) );
+								const FText DialogTitle( LOCTEXT("DatasmithConsumerAlreadyLoaded_DlgTitle", "Warning - Level already loaded") );
+
+								if(FMessageDialog::Open(EAppMsgType::YesNo, WarningMessage, &DialogTitle) != EAppReturnType::Yes)
+								{
+									bUnloadLevel = false;
+								}
+							}
+							else
+							{
+								const FTextFormat Format(LOCTEXT("DatasmithConsumer_AlreadyLoaded_Overwrite", "Level {0} with a different path from the same Dataprep asset is already loaded.It will be unloaded."));
+								const FText WarningMessage = FText::Format( Format, FText::FromString(LevelName));
+								LogWarning(WarningMessage);
+							}
 
-				// Indicates level's package is good for garbage collect
-				LevelPackage->SetDirtyFlag(false);
-				LevelPackage->ClearFlags(RF_Standalone | RF_Public | RF_Transactional);
-				LevelPackage->SetFlags(RF_Transient);
-				LevelPackage->MarkPendingKill();
+							if(bUnloadLevel)
+							{
+								WorkingWorld->RemoveLevel(Level);
+								if(ULevelStreaming* StreamingLevel = WorkingWorld->GetLevelStreamingForPackageName(*LevelWorldPath.GetLongPackageName()))
+								{
+									WorkingWorld->RemoveStreamingLevel(StreamingLevel);
+									WorkingWorld->UpdateLevelStreaming();
+								}
+							}
+						}
+					}
+				}
+				else
+				{
+					PrimaryLevel = Level;
+				}
+			}
+			else
+			{
+				check(false);
 			}
-		}
-
-		World->ClearStreamingLevels();
 
-		// Collect garbage to clear out the discarded world(s) and level(s) 
-		CollectGarbage( GARBAGE_COLLECTION_KEEPFLAGS );
+			break;
+		}
 	}
 
-	World->PersistentLevel = PrimaryLevel;
-	World->SetCurrentLevel(PrimaryLevel);
+	if(PrimaryLevel == nullptr)
+	{
+		PrimaryLevel = FindOrAddLevel(LevelName);
+		ensure(PrimaryLevel);
+	}
 
-	WorkingWorld = TStrongObjectPtr<UWorld>(World);
+	OriginalCurrentLevel = WorkingWorld->GetCurrentLevel();
+	WorkingWorld->SetCurrentLevel(PrimaryLevel);
 
 	return true;
 }
@@ -706,19 +620,8 @@ void UDatasmithConsumer::ClearWorld()
 {
 	if(WorkingWorld.IsValid())
 	{
-		UWorld* WorldToDelete = WorkingWorld.Get();
+		WorkingWorld->SetCurrentLevel(OriginalCurrentLevel);
 		WorkingWorld.Reset();
-
-		WorldToDelete->PersistentLevel = PrimaryLevel;
-		WorldToDelete->SetCurrentLevel(PrimaryLevel);
-
-		PrimaryLevel = nullptr;
-
-		WorldToDelete->DestroyWorld( true );
-		GEngine->DestroyWorldContext( WorldToDelete );
-
-		// Collect garbage to clear out the destroyed level
-		CollectGarbage( GARBAGE_COLLECTION_KEEPFLAGS );
 	}
 }
 
@@ -854,38 +757,69 @@ bool UDatasmithConsumer::CanCreateLevel(const FString& RequestedFolder, const FS
 {
 	FSoftObjectPath ObjectPath(FPaths::Combine(RequestedFolder, RequestedName) + TEXT(".") + RequestedName);
 
-	FString PackageFilename;
-	FPackageName::TryConvertLongPackageNameToFilename( ObjectPath.GetLongPackageName(), PackageFilename, FPackageName::GetMapPackageExtension() );
+	const FString AssetPathName = ObjectPath.GetAssetPathString();
 
-	if(FPaths::FileExists(PackageFilename))
+	if(FDatasmithImporterUtils::CanCreateAsset(AssetPathName, UWorld::StaticClass()) == FDatasmithImporterUtils::EAssetCreationStatus::CS_CanCreate)
 	{
-		if(UWorld* World = Cast<UWorld>(ObjectPath.TryLoad()))
+		if(FDatasmithImporterImpl::CheckAssetPersistenceValidity(FPaths::GetPath(AssetPathName), *ImportContextPtr))
 		{
-			if(DatasmithConsumerUtils::GetMarker(World->PersistentLevel, ConsumerMarkerID) != UniqueID)
+			FString PackageFilename;
+			FPackageName::TryConvertLongPackageNameToFilename( ObjectPath.GetLongPackageName(), PackageFilename, FPackageName::GetMapPackageExtension() );
+
+			if(FPaths::FileExists(PackageFilename))
 			{
-				if(bShowDialog)
+				if(UWorld* World = Cast<UWorld>(ObjectPath.TryLoad()))
 				{
-					const FTextFormat Format(LOCTEXT("DatasmithConsumer_SetTargetContentFolder_Overwrite_Dlg", "Level {0} already exists in {1}.\n\nDo you want to overwrite it?"));
-					const FText WarningMessage = FText::Format( Format, FText::FromString(RequestedName), FText::FromString(RequestedFolder));
-					const FText DialogTitle( LOCTEXT("DatasmithConsumer_Overwrite_DlgTitle", "Warning - Level already exists") );
-
-					if(FMessageDialog::Open(EAppMsgType::YesNo, WarningMessage, &DialogTitle) != EAppReturnType::Yes)
+					if(DatasmithConsumerUtils::GetMarker(World->PersistentLevel, ConsumerMarkerID) != UniqueID)
 					{
-						return false;
+						if(bShowDialog)
+						{
+							const FTextFormat Format(LOCTEXT("DatasmithConsumer_SetTargetContentFolder_Update_Dlg", "Level {0} already exists in {1} but is not from this Dataprep asset.\n\nDo you want to update it?"));
+							const FText WarningMessage = FText::Format( Format, FText::FromString(RequestedName), FText::FromString(RequestedFolder));
+							const FText DialogTitle( LOCTEXT("DatasmithConsumer_Update_DlgTitle", "Warning - Level already exists") );
+
+							EAppReturnType::Type Result = FMessageDialog::Open(EAppMsgType::YesNo, EAppReturnType::No, WarningMessage, &DialogTitle);
+
+							if(Result != EAppReturnType::Yes)
+							{
+								return false;
+							}
+						}
+						else
+						{
+							const FTextFormat Format(LOCTEXT("DatasmithConsumer_SetTargetContentFolder_Overwrite", "Level {0} already exists in {1} but is not from this Dataprep asset.It will be updated."));
+							const FText WarningMessage = FText::Format( Format, FText::FromString(RequestedName), FText::FromString(RequestedFolder));
+							LogWarning(WarningMessage);
+						}
 					}
 				}
-				else
-				{
-					const FTextFormat Format(LOCTEXT("DatasmithConsumer_SetTargetContentFolder_Overwrite", "Level {0} already exists in {1}.It will be overwritten."));
-					const FText WarningMessage = FText::Format( Format, FText::FromString(RequestedName), FText::FromString(RequestedFolder));
-					LogWarning(WarningMessage);
-				}
 			}
 		}
+		else
+		{
+			FString PackageFilename;
+			FPackageName::TryConvertLongPackageNameToFilename( ObjectPath.GetLongPackageName(), PackageFilename, FPackageName::GetMapPackageExtension() );
+
+			const FTextFormat Format(LOCTEXT("DatasmithConsumer_SetTargetContentFolder_CantCreateFile_Dlg", "Cannot create map file associated with level {0} in directory {1}.\nPlease choose another name for the level or fix the creation issue."));
+			const FText Message = FText::Format( Format, FText::FromString(RequestedName), FText::FromString( FPaths::ConvertRelativePathToFull(FPaths::GetPath(PackageFilename)) ));
+
+			if(bShowDialog)
+			{
+				const FText DialogTitle( LOCTEXT("DatasmithConsumer_CantCreateFile_DlgTitle", "Warning - Cannot create level") );
+
+				FMessageDialog::Open(EAppMsgType::Ok, Message, &DialogTitle);
+			}
+			else
+			{
+				LogError(Message);
+			}
+
+			return false;
+		}
 	}
-	else if(FDatasmithImporterUtils::CanCreateAsset(ObjectPath.GetAssetPathString(), UWorld::StaticClass()) != FDatasmithImporterUtils::EAssetCreationStatus::CS_CanCreate)
+	else
 	{
-		const FTextFormat Format(LOCTEXT("DatasmithConsumer_SetTargetContentFolder_CantCreate_Dlg", "Cannot create level {0} in folder {1}."));
+		const FTextFormat Format(LOCTEXT("DatasmithConsumer_SetTargetContentFolder_CantCreate_Dlg", "Cannot create level {0} in folder {1}.\nAn asset of different class already exists in this folder with the same name."));
 		const FText Message = FText::Format( Format, FText::FromString(RequestedName), FText::FromString(RequestedFolder));
 
 		if(bShowDialog)
@@ -893,13 +827,13 @@ bool UDatasmithConsumer::CanCreateLevel(const FString& RequestedFolder, const FS
 			const FText DialogTitle( LOCTEXT("DatasmithConsumer_CantCreate_DlgTitle", "Warning - Cannot create level") );
 
 			FMessageDialog::Open(EAppMsgType::Ok, Message, &DialogTitle);
-
-			return false;
 		}
 		else
 		{
 			LogError(Message);
 		}
+
+		return false;
 	}
 
 	return true;
@@ -907,6 +841,12 @@ bool UDatasmithConsumer::CanCreateLevel(const FString& RequestedFolder, const FS
 
 bool UDatasmithConsumer::SetTargetContentFolderImplementation(const FString& InTargetContentFolder, FText& OutFailureReason, const bool bIsAutomated)
 {
+	if(InTargetContentFolder == TargetContentFolder)
+	{
+		// #ueent_todo: This is weird it happens in some cases to investigate
+		return true;
+	}
+
 	if(!CanCreateLevel(InTargetContentFolder, LevelName, !bIsAutomated && !IsRunningCommandlet()))
 	{
 		return false;
@@ -967,7 +907,16 @@ ULevel* UDatasmithConsumer::FindOrAddLevel(const FString& InLevelName)
 
 	if(ULevelStreaming* StreamingLevel = FLevelUtils::FindStreamingLevel(WorkingWorld.Get(), *LevelPackageName))
 	{
-		return StreamingLevel->GetLoadedLevel();
+		if(ULevel* Level = StreamingLevel->GetLoadedLevel())
+		{
+			return Level;
+		}
+		else
+		{
+			WorkingWorld->LoadSecondaryLevels();
+			ensure(StreamingLevel->GetLoadedLevel());
+			return StreamingLevel->GetLoadedLevel();
+		}
 	}
 
 	ULevel* CurrentLevel = WorkingWorld->PersistentLevel;
@@ -976,8 +925,6 @@ ULevel* UDatasmithConsumer::FindOrAddLevel(const FString& InLevelName)
 	FString PackageFilename;
 	FPackageName::TryConvertLongPackageNameToFilename( LevelPackageName, PackageFilename, FPackageName::GetMapPackageExtension() );
 
-	bool bCleanLevel = false;
-
 	ULevelStreaming* StreamingLevel = nullptr;
 	if(FPaths::FileExists(PackageFilename))
 	{
@@ -993,8 +940,6 @@ ULevel* UDatasmithConsumer::FindOrAddLevel(const FString& InLevelName)
 		{
 			ensure(false);
 		}
-
-		bCleanLevel = true;
 	}
 	else
 	{
@@ -1012,28 +957,6 @@ ULevel* UDatasmithConsumer::FindOrAddLevel(const FString& InLevelName)
 		{
 			DatasmithConsumerUtils::SetMarker(NewLevel, ConsumerMarkerID, UniqueID);
 
-			if(bCleanLevel)
-			{
-				// Clean up the level if it contains actors from previous execution
-				UWorld* LevelWorld = Cast<UWorld>(NewLevel->GetOuter());
-
-				TArray<AActor*> LevelActors = NewLevel->Actors;
-				for(AActor* Actor : LevelActors)
-				{
-
-					if(Actor && Actor->GetRootComponent() && !Actor->IsA<AWorldSettings>() && !Actor->IsA<APhysicsVolume>() && !Actor->IsA<ABrush>() )
-					{
-						LevelWorld->DestroyActor(Actor, true);
-						Actor->UObject::Rename(nullptr, GetTransientPackage(), REN_DontCreateRedirectors | REN_ForceNoResetLoaders);
-					}
-				}
-
-				LevelWorld->CleanupWorld(false, false);
-				LevelWorld->CleanupActors();
-
-				DatasmithConsumerUtils::SaveMap(LevelWorld);
-			}
-
 			WorkingWorld->AddLevel(NewLevel);
 
 			return NewLevel;
@@ -1050,9 +973,14 @@ ULevel* UDatasmithConsumer::FindOrAddLevel(const FString& InLevelName)
 
 bool UDatasmithConsumer::CheckOutputDirectives()
 {
-	auto CanCreateAsset = [](const FString& AssetPathName,const UClass* AssetClass)
+	auto CanCreateAsset = [ImportContext = ImportContextPtr.Get()](const FString& AssetPathName,const UClass* AssetClass)
 	{
-		return FDatasmithImporterUtils::CanCreateAsset(AssetPathName, AssetClass) == FDatasmithImporterUtils::EAssetCreationStatus::CS_CanCreate;
+		if(FDatasmithImporterUtils::CanCreateAsset(AssetPathName, AssetClass) == FDatasmithImporterUtils::EAssetCreationStatus::CS_CanCreate)
+		{
+			return FDatasmithImporterImpl::CheckAssetPersistenceValidity(FPaths::GetPath(AssetPathName), *ImportContext);
+		}
+
+		return false;
 	};
 
 	const bool bShowDialog = !Context.bSilentMode && !IsRunningCommandlet();
@@ -1060,104 +988,8 @@ bool UDatasmithConsumer::CheckOutputDirectives()
 	// Collect garbage to clear out the destroyed level
 	CollectGarbage( GARBAGE_COLLECTION_KEEPFLAGS );
 
-	// Collect all sub-levels to be created
-	TSet<FString> LevelsToCreate;
-	LevelsToCreate.Add(LevelName);
-
-	for(AActor* Actor : Context.WorldPtr->GetCurrentLevel()->Actors)
-	{
-		if(Actor)
-		{
-			const FString& OutputLevelName = DatasmithConsumerUtils::GetMarker(Actor->GetRootComponent(), UDataprepContentConsumer::RelativeOutput);
-			if(OutputLevelName.Len() > 0)
-			{
-				LevelsToCreate.Add(OutputLevelName);
-			}
-		}
-	}
-
 	bool bCannotCreateAsset = false;
 
-	// check if any of the levels to create or update is not opened in the level editor
-	TArray<FString> OpenedLevels;
-	UWorld* GlobalWorld = GWorld;
-	{
-		TArray<ULevel*> WorldLevels = GlobalWorld->GetLevels();
-		for(ULevel* Level : WorldLevels)
-		{
-			if(Level)
-			{
-				const FString WorldLevelName = Level->GetOuter()->GetName();
-				if(LevelsToCreate.Contains(WorldLevelName))
-				{
-					OpenedLevels.Add(WorldLevelName);
-				}
-			}
-		}
-	}
-
-	for(const FString& LevelToCreate : LevelsToCreate)
-	{
-		FSoftObjectPath AssetSoftObjectPath(FPaths::Combine(TargetContentFolder, LevelToCreate) + "." + LevelToCreate);
-
-		if(!CanCreateAsset(AssetSoftObjectPath.GetAssetPathString(), UWorld::StaticClass()))
-		{
-			const FTextFormat TextFormat(LOCTEXT( "DatasmithConsumer_CannotCreateAsset", "Cannot create asset {0}. Commit will be aborted" ));
-			const FText Message = FText::Format( TextFormat, FText::FromString(AssetSoftObjectPath.GetAssetPathString()) );
-			LogError(Message);
-
-			bCannotCreateAsset = true;
-		}
-		// Check if umap file does not already exist. If so, user will be asked if he/she wants to overwrite it
-		else
-		{
-			FString PackageFilename;
-			FPackageName::TryConvertLongPackageNameToFilename( AssetSoftObjectPath.GetLongPackageName(), PackageFilename, FPackageName::GetMapPackageExtension() );
-			if(FPaths::FileExists(PackageFilename))
-			{
-				const FTextFormat TextFormat(LOCTEXT( "DatasmithConsumer_UMapAlreadyExists", "Level {0} already exists" ));
-				const FText Message = FText::Format( TextFormat, FText::FromString(LevelToCreate) );
-				LogInfo(Message);
-			}
-		}
-	}
-
-	// Abort commit if any level is opened
-	if(OpenedLevels.Num() > 0)
-	{
-		FText OpenedLevelsText;
-		if(OpenedLevels.Num() == 1)
-		{
-			const FTextFormat TextFormat(LOCTEXT( "DatasmithConsumer_OneLevelOpened", "level {0} is opened" ));
-			OpenedLevelsText = FText::Format( TextFormat, FText::FromString(OpenedLevels[0]) );
-		}
-		else
-		{
-			const FString Separator(TEXT(", "));
-			FString OpenedLevelListString = OpenedLevels[0];
-			for(int32 Index = 1; Index < OpenedLevels.Num(); ++Index)
-			{
-				OpenedLevelListString += Separator + OpenedLevels[Index];
-			}
-
-			const FTextFormat TextFormat(LOCTEXT( "DatasmithConsumer_MultipleLevelOpened", "levels {0} are opened"));
-			OpenedLevelsText = FText::Format( TextFormat, FText::FromString(OpenedLevelListString) );
-		}
-
-		const FText Message = FText::Format( LOCTEXT( "DatasmithConsumer_OpenAbortCommit", "Cannot proceed with commit because {0}.\nPlease close any editor using this level and commit again" ), OpenedLevelsText );
-		if(bShowDialog)
-		{
-			const FText Title( LOCTEXT( "DatasmithConsumer_OpenAbortCommitTitle", "Main level is opened" ) );
-			FMessageDialog::Open( EAppMsgType::Ok, Message, &Title );
-		}
-		else
-		{
-			LogError(Message);
-		}
-
-		return false;
-	}
-	
 	for(const TWeakObjectPtr< UObject >& AssetPtr : Context.Assets)
 	{
 		if(UObject* Asset = AssetPtr.Get())
@@ -1701,7 +1533,7 @@ namespace DatasmithConsumerUtils
 		const bool bIsMove = bShoudCut;
 		GEditor->CopySelectedActorsToClipboard(OwningWorld, bShoudCut, bIsMove, /*bWarnAboutReferences =*/ false);
 
-		UEditorLevelUtils::SetLevelVisibility(DestLevel, true, false);
+		UEditorLevelUtils::SetLevelVisibility(DestLevel, true, false, ELevelVisibilityDirtyMode::DontModify);
 
 		// Scope this so that Actors that have been pasted will have their final levels set before doing the actor mapping
 		{
diff --git a/Engine/Plugins/Enterprise/DatasmithImporter/Source/DatasmithImporter/Private/DatasmithConsumer.h b/Engine/Plugins/Enterprise/DatasmithImporter/Source/DatasmithImporter/Private/DatasmithConsumer.h
--- a/Engine/Plugins/Enterprise/DatasmithImporter/Source/DatasmithImporter/Private/DatasmithConsumer.h
+++ b/Engine/Plugins/Enterprise/DatasmithImporter/Source/DatasmithImporter/Private/DatasmithConsumer.h
@@ -104,14 +104,12 @@ private:
 	/** World used by the consumer to create new output levels */
 	TStrongObjectPtr<UWorld> WorkingWorld;
 
-	/** Indicates if LevelWorld was loaded specifically to execute this consumer */
-	bool bDeleteLevelWorld;
-
-	/** Levels which were loaded in the LevelEditor's world and need to be restored */
-	TArray<FString> LevelsToRestore;
-
 	UPROPERTY()
 	FSoftObjectPath OutputLevelSoftObject;
 
+	/** Level associated with the consumer */
 	ULevel* PrimaryLevel;
+
+	/** Current level used by the global world before executing the consumer */
+	ULevel* OriginalCurrentLevel;
 };
\ No newline at end of file
diff --git a/Engine/Plugins/Enterprise/DatasmithImporter/Source/DatasmithImporter/Private/DatasmithImporter.cpp b/Engine/Plugins/Enterprise/DatasmithImporter/Source/DatasmithImporter/Private/DatasmithImporter.cpp
--- a/Engine/Plugins/Enterprise/DatasmithImporter/Source/DatasmithImporter/Private/DatasmithImporter.cpp
+++ b/Engine/Plugins/Enterprise/DatasmithImporter/Source/DatasmithImporter/Private/DatasmithImporter.cpp
@@ -24,6 +24,7 @@
 #include "DatasmithAnimationElements.h"
 #include "LevelVariantSets.h"
 #include "ObjectTemplates/DatasmithObjectTemplate.h"
+#include "Utility/DatasmithImporterImpl.h"
 #include "Utility/DatasmithImporterUtils.h"
 #include "Utility/DatasmithTextureResize.h"
 
@@ -74,885 +75,6 @@ extern UNREALED_API UEditorEngine* GEditor;
 
 #define LOCTEXT_NAMESPACE "DatasmithImporter"
 
-namespace DatasmithImporterImpl
-{
-	static void ReportProgress(FScopedSlowTask* SlowTask, const float ExpectedWorkThisFrame, const FText& Text)
-	{
-		if (SlowTask)
-		{
-			SlowTask->EnterProgressFrame(ExpectedWorkThisFrame, Text);
-		}
-	}
-
-	static void ReportProgress(FScopedSlowTask* SlowTask, const float ExpectedWorkThisFrame, FText&& Text)
-	{
-		if (SlowTask)
-		{
-			SlowTask->EnterProgressFrame(ExpectedWorkThisFrame, Text);
-		}
-	}
-
-	static bool HasUserCancelledTask(FFeedbackContext* FeedbackContext)
-	{
-		if (FeedbackContext)
-		{
-			return FeedbackContext->ReceivedUserCancel();
-		}
-
-		return false;
-	}
-
-	static UObject* PublicizeAsset( UObject* SourceAsset, const TCHAR* DestinationPath, UObject* ExistingAsset )
-	{
-		TRACE_CPUPROFILER_EVENT_SCOPE(DatasmithImporterImpl::PublicizeAsset);
-
-		UPackage* DestinationPackage;
-
-		if ( !ExistingAsset )
-		{
-			const FString AssetName = SourceAsset->GetName();
-			bool bPathIsComplete = AssetName == FPaths::GetBaseFilename( DestinationPath );
-			FString DestinationPackagePath = UPackageTools::SanitizePackageName( bPathIsComplete ? DestinationPath : FPaths::Combine( DestinationPath, AssetName ) );
-			FString DestinationAssetPath = DestinationPackagePath + TEXT(".") + UPackageTools::SanitizePackageName( AssetName );
-
-			ExistingAsset = FDatasmithImporterUtils::FindObject<UObject>( nullptr, DestinationAssetPath );
-
-			DestinationPackage = ExistingAsset ? ExistingAsset->GetOutermost() : CreatePackage( nullptr, *DestinationPackagePath );
-		}
-		else
-		{
-			DestinationPackage = ExistingAsset->GetOutermost();
-		}
-
-		DestinationPackage->FullyLoad();
-
-		UObject* DestinationAsset = ExistingAsset;
-
-		FString OldAssetPathName;
-
-		// If the object already exist, then we need to fix up the reference
-		if ( ExistingAsset != nullptr && ExistingAsset != SourceAsset )
-		{
-			OldAssetPathName = ExistingAsset->GetPathName();
-
-			DestinationAsset = FDatasmithImporterUtils::DuplicateObject( SourceAsset, DestinationPackage, ExistingAsset->GetFName() );
-
-			// If mesh's label has changed, update its name
-			if ( ExistingAsset->GetFName() != SourceAsset->GetFName() )
-			{
-				DestinationAsset->Rename( *SourceAsset->GetName(), DestinationPackage, REN_DontCreateRedirectors | REN_NonTransactional );
-			}
-
-			if ( UStaticMesh* DestinationMesh = Cast< UStaticMesh >( DestinationAsset ) )
-			{
-				// This is done during the static mesh build process but we need to redo it after the DuplicateObject since the links are now valid
-				for ( TObjectIterator< UStaticMeshComponent > It; It; ++It )
-				{
-					if ( It->GetStaticMesh() == DestinationMesh )
-					{
-						It->FixupOverrideColorsIfNecessary( true );
-						It->InvalidateLightingCache();
-					}
-				}
-			}
-		}
-		else
-		{
-			SourceAsset->Rename( *SourceAsset->GetName(), DestinationPackage, REN_DontCreateRedirectors | REN_NonTransactional );
-			DestinationAsset = SourceAsset;
-		}
-
-		DestinationAsset->SetFlags( RF_Public );
-		DestinationAsset->MarkPackageDirty();
-
-		if ( !ExistingAsset )
-		{
-			FAssetRegistryModule::AssetCreated( DestinationAsset );
-		}
-		else if ( !OldAssetPathName.IsEmpty() )
-		{
-			FAssetRegistryModule::AssetRenamed( DestinationAsset, OldAssetPathName );
-		}
-
-		return DestinationAsset;
-	}
-
-	/**
-	 * Verifies the input asset can successfully be saved and/or cooked.
-	 * The code below is heavily inspired from:
-	 *			- ContentBrowserUtils::IsValidPackageForCooking
-	 * #ueent_todo Make ContentBrowserUtils::IsValidPackageForCooking public. Not game to do it so late in 4.22
-	 */
-	static void CheckAssetPersistenceValidity(UObject* Asset, FDatasmithImportContext& ImportContext)
-	{
-		if (Asset == nullptr)
-		{
-			return;
-		}
-
-		UPackage* Package = Asset->GetOutermost();
-		const FString PackageName = Package->GetName();
-
-		// Check that package can be saved
-		const FString BasePackageFileName = FPackageName::LongPackageNameToFilename( PackageName );
-		const FString AbsolutePathToAsset = FPaths::ConvertRelativePathToFull( BasePackageFileName );
-
-		// Create fake filename of same length of final asset file name to test ability to write
-		const FString FakeAbsolutePathToAsset = AbsolutePathToAsset + TEXT( ".uasset" );
-
-		// Verify asset file name does not exceed OS' maximum path length
-		if( FPlatformMisc::GetMaxPathLength() < FakeAbsolutePathToAsset.Len() )
-		{
-			ImportContext.LogWarning(FText::Format(LOCTEXT("DatasmithImportInvalidLength", "Saving may partially fail because path for asset {0} is too long. Rename before saving."), FText::FromString( PackageName )));
-		}
-		// Verify user can overwrite existing file
-		else if( IFileManager::Get().FileExists( *FakeAbsolutePathToAsset ) )
-		{
-			FFileStatData FileStatData = IFileManager::Get().GetStatData( *FakeAbsolutePathToAsset );
-			if ( FileStatData.bIsReadOnly )
-			{
-				// Check to see if the file is not under source control
-				bool bWarnUser = true;
-
-				ISourceControlProvider& SourceControlProvider = ISourceControlModule::Get().GetProvider();
-				if (SourceControlProvider.IsAvailable() && SourceControlProvider.IsEnabled())
-				{
-					SourceControlProvider.Execute(ISourceControlOperation::Create<FUpdateStatus>(), FakeAbsolutePathToAsset);
-					FSourceControlStatePtr SourceControlState = SourceControlProvider.GetState(FakeAbsolutePathToAsset, EStateCacheUsage::Use);
-					if( SourceControlState.IsValid() && SourceControlState->CanCheckout() )
-					{
-						// User will be prompted to check out this file when he/she saves the asset. No need to warn.
-						bWarnUser = false;
-					}
-				}
-
-				if(bWarnUser)
-				{
-					ImportContext.LogWarning(FText::Format(LOCTEXT("DatasmithImportInvalidSaving", "Saving may partially fail because file asset {0} cannot be overwritten. Check your privileges."), FText::FromString( PackageName )));
-				}
-			}
-		}
-		// Verify user has privileges to write in folder where asset file will be stored
-		else
-		{
-			// We can't just check for the target content folders with IFileManager::Get().GetStatData here as those will
-			// only be created when UUnrealEdEngine::GetWarningStateForWritePermission is called to check for
-			// write permissions the first time, as the result is cached in GUnrealEd->PackagesCheckedForWritePermission.
-			// To check for permission, we need to first check this cache, and if the PackageName hasn't been
-			// checked yet, we need to replicate what UUnrealEdEngine::GetWarningStateForWritePermission does
-			EWriteDisallowedWarningState WarningState = EWriteDisallowedWarningState::WDWS_MAX;
-			if (GUnrealEd != nullptr && GUnrealEd->PackagesCheckedForWritePermission.Find(PackageName))
-			{
-				WarningState = (EWriteDisallowedWarningState)*GUnrealEd->PackagesCheckedForWritePermission.Find(PackageName);
-			}
-			else if (FFileHelper::SaveStringToFile(TEXT("Write Test"), *FakeAbsolutePathToAsset))
-			{
-				// We can successfully write to the folder containing the package.
-				// Delete the temp file.
-				IFileManager::Get().Delete(*FakeAbsolutePathToAsset);
-				WarningState = EWriteDisallowedWarningState::WDWS_WarningUnnecessary;
-			}
-
-			if(WarningState != EWriteDisallowedWarningState::WDWS_WarningUnnecessary)
-			{
-				ImportContext.LogWarning(FText::Format(LOCTEXT("DatasmithImportInvalidFolder", "Cannot write in folder {0} to store asset {1}. Check access to folder."), FText::FromString( FPaths::GetPath( FakeAbsolutePathToAsset ) ), FText::FromString( PackageName )));
-			}
-		}
-
-		// Check that package can be cooked
-		// Value for MaxGameNameLen directly taken from ContentBrowserUtils::GetPackageLengthForCooking
-		constexpr int32 MaxGameNameLen = 20;
-
-		// Pad out the game name to the maximum allowed
-		const FString GameName = FApp::GetProjectName();
-		FString GameNamePadded = GameName;
-		while (GameNamePadded.Len() < MaxGameNameLen)
-		{
-			GameNamePadded += TEXT(" ");
-		}
-
-		const FString AbsoluteGamePath = FPaths::ConvertRelativePathToFull(FPaths::ProjectDir());
-		const FString AbsoluteGameCookPath = AbsoluteGamePath / TEXT("Saved") / TEXT("Cooked") / TEXT("WindowsNoEditor") / GameName;
-
-		FString AssetPathWithinCookDir = AbsolutePathToAsset;
-		FPaths::RemoveDuplicateSlashes(AssetPathWithinCookDir);
-		AssetPathWithinCookDir.RemoveFromStart(AbsoluteGamePath, ESearchCase::CaseSensitive);
-
-		// Test that the package can be cooked based on the current project path
-		FString AbsoluteCookPathToAsset = AbsoluteGameCookPath / AssetPathWithinCookDir;
-
-		AbsoluteCookPathToAsset.ReplaceInline(*GameName, *GameNamePadded, ESearchCase::CaseSensitive);
-
-		// Get the longest path allowed by the system or use 260 as the longest which is the shortest max path of any platforms that support cooking
-		const int32 MaxCookPath = GetDefault<UEditorExperimentalSettings>()->bEnableLongPathsSupport ? FPlatformMisc::GetMaxPathLength() : 260 /* MAX_PATH */;
-
-		if (AbsoluteCookPathToAsset.Len() > MaxCookPath)
-		{
-			ImportContext.LogWarning(FText::Format(LOCTEXT("DatasmithImportInvalidCooking", "Cooking may fail because path for asset {0} is too long. Rename before cooking."), FText::FromString(PackageName)));
-		}
-	}
-
-	/** Set the texture mode on each texture element based on its usage in the materials */
-	static void SetTexturesMode( FDatasmithImportContext& ImportContext )
-	{
-		const int32 TexturesCount = ImportContext.FilteredScene->GetTexturesCount();
-		const int32 MaterialsCount = ImportContext.FilteredScene->GetMaterialsCount();
-
-		FFeedbackContext* FeedbackContext = ImportContext.FeedbackContext;
-		for ( int32 TextureIndex = 0; TextureIndex < TexturesCount && !ImportContext.bUserCancelled; ++TextureIndex )
-		{
-			ImportContext.bUserCancelled |= DatasmithImporterImpl::HasUserCancelledTask( ImportContext.FeedbackContext );
-
-			TSharedPtr< IDatasmithTextureElement > TextureElement = ImportContext.FilteredScene->GetTexture( TextureIndex );
-			const FString TextureName = ObjectTools::SanitizeObjectName( TextureElement->GetName() );
-
-			for ( int32 MaterialIndex = 0; MaterialIndex < MaterialsCount; ++MaterialIndex )
-			{
-				const TSharedPtr< IDatasmithBaseMaterialElement >& BaseMaterialElement = ImportContext.FilteredScene->GetMaterial( MaterialIndex );
-
-				if ( BaseMaterialElement->IsA( EDatasmithElementType::Material ) )
-				{
-					const TSharedPtr< IDatasmithMaterialElement >& MaterialElement = StaticCastSharedPtr< IDatasmithMaterialElement >( BaseMaterialElement );
-
-					for ( int32 s = 0; s < MaterialElement->GetShadersCount(); ++s )
-					{
-						const TSharedPtr< IDatasmithShaderElement >& ShaderElement = MaterialElement->GetShader(s);
-
-						if ( FCString::Strlen( ShaderElement->GetDiffuseTexture() ) > 0 && ShaderElement->GetDiffuseTexture() == TextureName)
-						{
-							TextureElement->SetTextureMode(EDatasmithTextureMode::Diffuse);
-						}
-						else if ( FCString::Strlen( ShaderElement->GetReflectanceTexture() ) > 0 && ShaderElement->GetReflectanceTexture() == TextureName)
-						{
-							TextureElement->SetTextureMode(EDatasmithTextureMode::Specular);
-						}
-						else if ( FCString::Strlen( ShaderElement->GetDisplaceTexture() ) > 0 && ShaderElement->GetDisplaceTexture() == TextureName)
-						{
-							TextureElement->SetTextureMode(EDatasmithTextureMode::Displace);
-						}
-						else if ( FCString::Strlen( ShaderElement->GetNormalTexture() ) > 0 && ShaderElement->GetNormalTexture() == TextureName)
-						{
-							if (!ShaderElement->GetNormalTextureSampler().bInvert)
-							{
-								TextureElement->SetTextureMode(EDatasmithTextureMode::Normal);
-							}
-							else
-							{
-								TextureElement->SetTextureMode(EDatasmithTextureMode::NormalGreenInv);
-							}
-						}
-					}
-				}
-				else if ( BaseMaterialElement->IsA( EDatasmithElementType::UEPbrMaterial ) )
-				{
-					const TSharedPtr< IDatasmithUEPbrMaterialElement >& MaterialElement = StaticCastSharedPtr< IDatasmithUEPbrMaterialElement >( BaseMaterialElement );
-
-					TFunction< bool( IDatasmithMaterialExpression* ) > IsTextureConnected;
-					IsTextureConnected = [ &TextureName, &IsTextureConnected ]( IDatasmithMaterialExpression* MaterialExpression ) -> bool
-					{
-						if ( !MaterialExpression )
-						{
-							return false;
-						}
-
-						if ( MaterialExpression->IsA( EDatasmithMaterialExpressionType::Texture ) )
-						{
-							IDatasmithMaterialExpressionTexture* TextureExpression = static_cast< IDatasmithMaterialExpressionTexture* >( MaterialExpression );
-
-							if ( TextureExpression->GetTexturePathName() == TextureName )
-							{
-								return true;
-							}
-						}
-
-						for ( int32 InputIndex = 0; InputIndex < MaterialExpression->GetInputCount(); ++InputIndex )
-						{
-							IDatasmithMaterialExpression* ConnectedExpression = MaterialExpression->GetInput( InputIndex )->GetExpression();
-
-							if ( ConnectedExpression && IsTextureConnected( ConnectedExpression ) )
-							{
-								return true;
-							}
-						}
-
-						return false;
-					};
-
-					if ( IsTextureConnected( MaterialElement->GetBaseColor().GetExpression() ) )
-					{
-						TextureElement->SetTextureMode(EDatasmithTextureMode::Diffuse);
-					}
-					else if ( IsTextureConnected( MaterialElement->GetSpecular().GetExpression() ) )
-					{
-						TextureElement->SetTextureMode(EDatasmithTextureMode::Specular);
-					}
-					else if ( IsTextureConnected( MaterialElement->GetNormal().GetExpression() ) )
-					{
-						if ( TextureElement->GetTextureMode() != EDatasmithTextureMode::Bump )
-						{
-							TextureElement->SetTextureMode(EDatasmithTextureMode::Normal);
-						}
-					}
-				}
-			}
-		}
-	}
-
-	static void CompileMaterial( UObject* Material )
-	{
-		if ( !Material->IsA< UMaterialInterface >() && !Material->IsA< UMaterialFunctionInterface >() )
-		{
-			return;
-		}
-
-		TRACE_CPUPROFILER_EVENT_SCOPE(DatasmithImporterImpl::CompileMaterial);
-
-		FMaterialUpdateContext MaterialUpdateContext;
-
-		if ( UMaterialInterface* MaterialInterface = Cast< UMaterialInterface >( Material ) )
-		{
-			MaterialUpdateContext.AddMaterialInterface( MaterialInterface );
-		}
-
-		if ( UMaterialInstanceConstant* ConstantMaterialInstance = Cast< UMaterialInstanceConstant >( Material ) )
-		{
-			// If BlendMode override property has been changed, make sure this combination of the parent material is compiled
-			if ( ConstantMaterialInstance->BasePropertyOverrides.bOverride_BlendMode == true )
-			{
-				ConstantMaterialInstance->ForceRecompileForRendering();
-			}
-			else
-			{
-				// If a static switch is overriden, we need to recompile
-				FStaticParameterSet StaticParameters;
-				ConstantMaterialInstance->GetStaticParameterValues( StaticParameters );
-
-				for ( FStaticSwitchParameter& Switch : StaticParameters.StaticSwitchParameters )
-				{
-					if ( Switch.bOverride )
-					{
-						ConstantMaterialInstance->ForceRecompileForRendering();
-						break;
-					}
-				}
-			}
-		}
-
-		Material->PreEditChange( nullptr );
-		Material->PostEditChange();
-	}
-
-	static void FixReferencesForObject( UObject* Object, const TMap< UObject*, UObject* >& ReferencesToRemap )
-	{
-		constexpr bool bNullPrivateRefs = false;
-		constexpr bool bIgnoreOuterRef = true;
-		constexpr bool bIgnoreArchetypeRef = true;
-
-		if ( ReferencesToRemap.Num() > 0 )
-		{
-			FArchiveReplaceObjectRef< UObject > ArchiveReplaceObjectRef( Object, ReferencesToRemap, bNullPrivateRefs, bIgnoreOuterRef, bIgnoreArchetypeRef );
-		}
-	}
-
-	using FMigratedTemplatePairType = TPair< TStrongObjectPtr< UDatasmithObjectTemplate >, TStrongObjectPtr< UDatasmithObjectTemplate > >;
-
-	/**
-	 * Creates templates to apply the values from the SourceObject on the DestinationObject.
-	 *
-	 * @returns An array of template pairs. The key is the template for the object, the value is a template to force apply to the object,
-	 *			it contains the values from the key and any overrides that were present on the DestinationObject.
-	 */
-	static TArray< FMigratedTemplatePairType > MigrateTemplates( UObject* SourceObject, UObject* DestinationObject, const TMap< UObject*, UObject* >* ReferencesToRemap, bool bIsForActor )
-	{
-		TArray< FMigratedTemplatePairType > Results;
-
-		if ( !SourceObject )
-		{
-			return Results;
-		}
-
-		TMap< TSubclassOf< UDatasmithObjectTemplate >, UDatasmithObjectTemplate* >* SourceTemplates = FDatasmithObjectTemplateUtils::FindOrCreateObjectTemplates( SourceObject );
-
-		if ( !SourceTemplates )
-		{
-			return Results;
-		}
-
-		for ( const TPair< TSubclassOf< UDatasmithObjectTemplate >, UDatasmithObjectTemplate* >& SourceTemplatePair : *SourceTemplates )
-		{
-			if ( bIsForActor == SourceTemplatePair.Value->bIsActorTemplate )
-			{
-				FMigratedTemplatePairType& Result = Results.AddDefaulted_GetRef();
-
-				TStrongObjectPtr< UDatasmithObjectTemplate > SourceTemplate{ NewObject< UDatasmithObjectTemplate >(GetTransientPackage(), SourceTemplatePair.Key.Get()) }; // The SourceTemplate is the one we will persist so set its outer as DestinationObject
-
-				SourceTemplate->Load(SourceObject);
-
-				if ( ReferencesToRemap )
-				{
-					FixReferencesForObject(SourceTemplate.Get(), *ReferencesToRemap);
-				}
-
-				Result.Key = SourceTemplate;
-
-				if (DestinationObject && !DestinationObject->IsPendingKillOrUnreachable())
-				{
-					Result.Value = TStrongObjectPtr< UDatasmithObjectTemplate >(UDatasmithObjectTemplate::GetDifference( DestinationObject, SourceTemplate.Get()));
-				}
-				else
-				{
-					Result.Value = SourceTemplate;
-				}
-			}
-		}
-
-		return Results;
-	}
-
-	/**
-	 * Applies the templates created from MigrateTemplates to DestinationObject.
-	 *
-	 * For an Object A that should be duplicated over an existing A', for which we want to keep the Datasmith overrides:
-	 * - Call MigrateTemplates(A, A')
-	 * - Duplicate A over A'
-	 * - ApplyMigratedTemplates(A')
-	 */
-	static void ApplyMigratedTemplates( TArray< FMigratedTemplatePairType >& MigratedTemplates, UObject* DestinationObject )
-	{
-		for ( FMigratedTemplatePairType& MigratedTemplate : MigratedTemplates )
-		{
-			UDatasmithObjectTemplate* SourceTemplate = MigratedTemplate.Key.Get();
-			UDatasmithObjectTemplate* DestinationTemplate = MigratedTemplate.Value.Get();
-
-			DestinationTemplate->Apply(DestinationObject, true); // Restore the overrides
-			FDatasmithObjectTemplateUtils::SetObjectTemplate(DestinationObject, SourceTemplate); // Set SourceTemplate as our template so that any differences are considered overrides
-
-		}
-	}
-
-	static UObject* FinalizeAsset( UObject* SourceAsset, const TCHAR* AssetPath, UObject* ExistingAsset, TMap< UObject*, UObject* >* ReferencesToRemap )
-	{
-		if ( ReferencesToRemap )
-		{
-			FixReferencesForObject( SourceAsset, *ReferencesToRemap );
-		}
-
-		TArray< FMigratedTemplatePairType > MigratedTemplates = MigrateTemplates( SourceAsset, ExistingAsset, ReferencesToRemap, false );
-
-		UObject* FinalAsset = PublicizeAsset( SourceAsset, AssetPath, ExistingAsset );
-
-		ApplyMigratedTemplates( MigratedTemplates, FinalAsset );
-
-		if ( ReferencesToRemap && SourceAsset && SourceAsset != FinalAsset )
-		{
-			ReferencesToRemap->Add( SourceAsset, FinalAsset );
-		}
-
-		return FinalAsset;
-	}
-
-	class FActorWriter final: public FObjectWriter
-	{
-	public:
-		FActorWriter( UObject* Object, TArray< uint8 >& Bytes )
-			: FObjectWriter( Bytes )
-		{
-			SetIsLoading(false);
-			SetIsSaving(true);
-			SetIsPersistent(false);
-
-			Object->Serialize(*this); // virtual call in ctr -> final class
-		}
-
-		virtual bool ShouldSkipProperty(const FProperty* InProperty) const override
-		{
-			bool bSkip = false;
-
-			if ( InProperty->IsA< FObjectPropertyBase >() )
-			{
-				bSkip = true;
-			}
-			else if ( InProperty->HasAnyPropertyFlags(CPF_Transient) || !InProperty->HasAnyPropertyFlags(CPF_Edit | CPF_Interp) )
-			{
-				bSkip = true;
-			}
-
-			return bSkip;
-		}
-	};
-
-	class FComponentWriter final : public FObjectWriter
-	{
-	public:
-		FComponentWriter( UObject* Object, TArray< uint8 >& Bytes )
-			: FObjectWriter( Bytes )
-		{
-			SetIsLoading(false);
-			SetIsSaving(true);
-			SetIsPersistent(false);
-
-			Object->Serialize(*this); // virtual call in ctr -> final class
-		}
-
-		virtual bool ShouldSkipProperty(const FProperty* InProperty) const override
-		{
-			bool bSkip = false;
-
-			if ( InProperty->HasAnyPropertyFlags(CPF_Transient) || !InProperty->HasAnyPropertyFlags(CPF_Edit | CPF_Interp) )
-			{
-				bSkip = true;
-			}
-
-			return bSkip;
-		}
-	};
-
-	static void DeleteImportSceneActorIfNeeded(FDatasmithActorImportContext& ActorContext, bool bForce = false)
-	{
-		ADatasmithSceneActor*& ImportSceneActor = ActorContext.ImportSceneActor;
-		if ( !ActorContext.FinalSceneActors.Contains(ImportSceneActor) || bForce )
-		{
-			if ( ImportSceneActor )
-			{
-
-				TArray< TSoftObjectPtr< AActor > > RelatedActors;
-				ImportSceneActor->RelatedActors.GenerateValueArray( RelatedActors );
-
-				ImportSceneActor->Scene = nullptr;
-				ImportSceneActor->RelatedActors.Empty();
-
-				while(RelatedActors.Num() > 0)
-				{
-					TSoftObjectPtr< AActor > ActorPtr = RelatedActors.Pop(false);
-					if(AActor* RelatedActor = ActorPtr.Get())
-					{
-						FDatasmithImporterUtils::DeleteActor( *RelatedActor );
-					}
-				}
-
-				FDatasmithImporterUtils::DeleteActor( *ImportSceneActor );
-
-				// Null also the ImportSceneActor from the Actor Context because it's a ref to it.
-				ImportSceneActor = nullptr;
-			}
-		}
-	}
-
-	static UActorComponent* PublicizeComponent(UActorComponent& SourceComponent, UActorComponent* DestinationComponent, AActor& DestinationActor, TMap< UObject*, UObject* >& ReferencesToRemap, USceneComponent* DestinationParent = nullptr )
-	{
-		TRACE_CPUPROFILER_EVENT_SCOPE(DatasmithImporterImpl::PublicizeComponent);
-
-		if ( !SourceComponent.HasAnyFlags( RF_Transient | RF_TextExportTransient ) )
-		{
-			if (!DestinationComponent || DestinationComponent->IsPendingKillOrUnreachable())
-			{
-				if (DestinationComponent)
-				{
-					// Change the name of the old component so that the new object won't recycle the old one.
-					DestinationComponent->Rename(nullptr, nullptr, REN_DontCreateRedirectors | REN_NonTransactional);
-				}
-
-				if ( UActorComponent* OldComponent = static_cast<UActorComponent*>( FindObjectWithOuter( &DestinationActor, UActorComponent::StaticClass(), SourceComponent.GetFName() ) ) )
-				{
-					OldComponent->DestroyComponent( true );
-					// Change the name of the old component so that the new object won't recycle the old one.
-					OldComponent->Rename(nullptr, nullptr, REN_DontCreateRedirectors | REN_NonTransactional);
-				}
-
-				// Create a new component
-				DestinationComponent = NewObject<UActorComponent>(&DestinationActor, SourceComponent.GetClass(), SourceComponent.GetFName(), RF_Transactional);
-				DestinationActor.AddInstanceComponent(DestinationComponent);
-
-				if (USceneComponent* NewSceneComponent = Cast<USceneComponent>(DestinationComponent))
-				{
-					if ( !DestinationActor.GetRootComponent() )
-					{
-						DestinationActor.SetRootComponent(NewSceneComponent);
-					}
-					if ( DestinationParent )
-					{
-						NewSceneComponent->AttachToComponent( DestinationParent, FAttachmentTransformRules::KeepRelativeTransform );
-					}
-				}
-
-				DestinationComponent->RegisterComponent();
-			}
-
-			check(DestinationComponent);
-
-			// Copy component data
-			{
-				TArray< uint8 > Bytes;
-				FComponentWriter ObjectWriter(&SourceComponent, Bytes);
-				FObjectReader ObjectReader(DestinationComponent, Bytes);
-			}
-
-			FixReferencesForObject(DestinationComponent, ReferencesToRemap);
-
-			// #ueent_todo: we shouldn't be copying instanced object pointers in the first place
-			UDatasmithAssetUserData* SourceAssetUserData = DestinationComponent->GetAssetUserData< UDatasmithAssetUserData >();
-
-			if ( SourceAssetUserData )
-			{
-				UDatasmithAssetUserData* DestinationAssetUserData = DuplicateObject< UDatasmithAssetUserData >( SourceAssetUserData, DestinationComponent );
-				DestinationComponent->RemoveUserDataOfClass( UDatasmithAssetUserData::StaticClass() );
-				DestinationComponent->AddAssetUserData( DestinationAssetUserData );
-			}
-
-			ReferencesToRemap.Add( &SourceComponent, DestinationComponent );
-
-			return DestinationComponent;
-		}
-
-		return nullptr;
-	}
-
-	static void FinalizeSceneComponent(FDatasmithImportContext& ImportContext, USceneComponent& SourceComponent, AActor& DestinationActor, USceneComponent* DestinationParent, TMap<UObject *, UObject *>& ReferencesToRemap )
-	{
-		USceneComponent* DestinationComponent = static_cast<USceneComponent*>( FindObjectWithOuter( &DestinationActor, SourceComponent.GetClass(), SourceComponent.GetFName() ) );
-		FName SourceComponentDatasmithId = FDatasmithImporterUtils::GetDatasmithElementId(&SourceComponent);
-
-		if ( SourceComponentDatasmithId.IsNone() )
-		{
-			// This component is not tracked by datasmith
-			if ( !DestinationComponent || DestinationComponent->IsPendingKillOrUnreachable() )
-			{
-				DestinationComponent = static_cast<USceneComponent*> ( PublicizeComponent(SourceComponent, DestinationComponent, DestinationActor, ReferencesToRemap, DestinationParent) );
-				if ( DestinationComponent )
-				{
-					// Put back the components in a proper state
-					DestinationComponent->UpdateComponentToWorld();
-				}
-			}
-		}
-		else
-		{
-			check(ImportContext.ActorsContext.CurrentTargetedScene);
-
-			TArray< FMigratedTemplatePairType > MigratedTemplates = MigrateTemplates(
-				&SourceComponent,
-				DestinationComponent,
-				&ReferencesToRemap,
-				false
-			);
-
-			DestinationComponent = static_cast< USceneComponent* > ( PublicizeComponent( SourceComponent, DestinationComponent, DestinationActor, ReferencesToRemap, DestinationParent ) );
-
-			if ( DestinationComponent )
-			{
-				// Put back the components in a proper state (Unfortunately without this the set relative transform might not work)
-				DestinationComponent->UpdateComponentToWorld();
-				ApplyMigratedTemplates( MigratedTemplates, DestinationComponent );
-				DestinationComponent->PostEditChange();
-			}
-		}
-
-		USceneComponent* AttachParentForChildren = DestinationComponent  ? DestinationComponent : DestinationParent;
-		for ( USceneComponent* Child : SourceComponent.GetAttachChildren() )
-		{
-			// Only finalize components that are from the same actor
-			if ( Child && Child->GetOuter() == SourceComponent.GetOuter() )
-			{
-				FinalizeSceneComponent( ImportContext, *Child, DestinationActor, AttachParentForChildren, ReferencesToRemap );
-			}
-		}
-	}
-
-	static void FinalizeComponents(FDatasmithImportContext& ImportContext, AActor& SourceActor, AActor& DestinationActor, TMap<UObject *, UObject *>& ReferencesToRemap)
-	{
-		USceneComponent* ParentComponent = nullptr;
-
-		// Find the parent component
-		UObject** ObjectPtr = ReferencesToRemap.Find( SourceActor.GetRootComponent()->GetAttachParent() );
-		if ( ObjectPtr )
-		{
-			ParentComponent = Cast<USceneComponent>( *ObjectPtr );
-		}
-
-		// Finalize the scene components recursively
-		{
-			USceneComponent* RootComponent = SourceActor.GetRootComponent();
-			if ( RootComponent )
-			{
-				FinalizeSceneComponent( ImportContext, *RootComponent, DestinationActor, ParentComponent, ReferencesToRemap );
-			}
-		}
-
-
-		for ( UActorComponent* SourceComponent : SourceActor.GetComponents() )
-		{
-			// Only the non scene component haven't been finalize
-			if ( SourceComponent && !SourceComponent->GetClass()->IsChildOf<USceneComponent>() )
-			{
-				UActorComponent* DestinationComponent = static_cast< UActorComponent* >( FindObjectWithOuter( &DestinationActor, SourceComponent->GetClass(), SourceComponent->GetFName() ) );
-				if (!DestinationComponent)
-				{
-					DestinationComponent = PublicizeComponent( *SourceComponent, DestinationComponent, DestinationActor, ReferencesToRemap );
-				}
-			}
-		}
-	}
-
-	static void GatherUnsupportedVirtualTexturesAndMaterials(const TMap<TSharedRef< IDatasmithBaseMaterialElement >, UMaterialInterface*>& ImportedMaterials, TSet<UTexture2D*>& VirtualTexturesToConvert, TArray<UMaterial*>& MaterialsToRefreshAfterVirtualTextureConversion)
-	{
-		//Multimap cache to avoid parsing the same base material multiple times.
-		TMultiMap<UMaterial*, FMaterialParameterInfo> TextureParametersToConvertMap;
-
-		//Loops through all imported material instances and add to VirtualTexturesToConvert all the virtual texture parameters that don't support virtual texturing in the base material.
-		for (const TPair< TSharedRef< IDatasmithBaseMaterialElement >, UMaterialInterface* >& ImportedMaterialPair : ImportedMaterials)
-		{
-			UMaterialInterface* CurrentMaterialInterface = ImportedMaterialPair.Value;
-			UMaterial* BaseMaterial = CurrentMaterialInterface->GetMaterial();
-
-			if(!TextureParametersToConvertMap.Contains(BaseMaterial))
-			{
-				bool bRequiresTextureCheck = false;
-				TArray<FMaterialParameterInfo> OutParameterInfo;
-				TArray<FGuid> Guids;
-				BaseMaterial->GetAllTextureParameterInfo(OutParameterInfo, Guids);
-
-				for (int32 ParameterInfoIndex = 0; ParameterInfoIndex < OutParameterInfo.Num(); ++ParameterInfoIndex)
-				{
-					UTexture* TextureParameter = nullptr;
-
-					if (BaseMaterial->GetTextureParameterValue(OutParameterInfo[ParameterInfoIndex], TextureParameter) && VirtualTexturesToConvert.Contains(Cast<UTexture2D>(TextureParameter)))
-					{
-						bRequiresTextureCheck = true;
-						TextureParametersToConvertMap.Add(BaseMaterial, OutParameterInfo[ParameterInfoIndex]);
-					}
-				}
-
-				if (bRequiresTextureCheck)
-				{
-					MaterialsToRefreshAfterVirtualTextureConversion.Add(BaseMaterial);
-				}
-				else
-				{
-					//Adding a dummy MaterialParameterInfo so that we don't have to parse this Base Material again.
-					TextureParametersToConvertMap.Add(BaseMaterial, FMaterialParameterInfo());
-
-					//If no unsupported virtual texture parameters were found, it's possible that a texture needing conversion is simply not exposed as a parameter, so we still need to check for those.
-					for (UObject* ReferencedTexture : BaseMaterial->GetCachedExpressionData().ReferencedTextures)
-					{
-						if (VirtualTexturesToConvert.Contains(Cast<UTexture2D>(ReferencedTexture)))
-						{
-							MaterialsToRefreshAfterVirtualTextureConversion.Add(BaseMaterial);
-							break;
-						}
-					}
-				}
-			}
-
-			for (auto ParameterInfoIt = TextureParametersToConvertMap.CreateKeyIterator(BaseMaterial); ParameterInfoIt; ++ParameterInfoIt)
-			{
-				UTexture* TextureParameter = nullptr;
-
-				if (CurrentMaterialInterface->GetTextureParameterValue(ParameterInfoIt.Value(), TextureParameter) &&
-					TextureParameter && TextureParameter->VirtualTextureStreaming)
-				{
-					if (UTexture2D* TextureToConvert = Cast<UTexture2D>(TextureParameter))
-					{
-						VirtualTexturesToConvert.Add(TextureToConvert);
-					}
-				}
-			}
-		}
-	}
-
-	static void ConvertUnsupportedVirtualTexture(FDatasmithImportContext& ImportContext, TSet<UTexture2D*>& VirtualTexturesToConvert, const TMap<UObject*, UObject*>& ReferencesToRemap)
-	{
-		TArray<UMaterial*> MaterialsToRefreshAfterVirtualTextureConversion;
-		GatherUnsupportedVirtualTexturesAndMaterials(ImportContext.ImportedMaterials, ImportContext.AssetsContext.VirtualTexturesToConvert, MaterialsToRefreshAfterVirtualTextureConversion);
-
-		if (VirtualTexturesToConvert.Num() != 0)
-		{
-			for (UTexture2D*& TextureToConvert : VirtualTexturesToConvert)
-			{
-				if (UObject* const* RemappedTexture = ReferencesToRemap.Find(TextureToConvert))
-				{
-					TextureToConvert = Cast<UTexture2D>(*RemappedTexture);
-				}
-
-				ImportContext.LogWarning(FText::Format(LOCTEXT("DatasmithVirtualTextureConverted", "The imported texture {0} could not be imported as virtual texture as it is not supported in all the materials using it."), FText::FromString(TextureToConvert->GetName())));
-			}
-
-			for (int32 MaterialIndex = 0; MaterialIndex < MaterialsToRefreshAfterVirtualTextureConversion.Num(); ++MaterialIndex)
-			{
-				if (UObject* const* RemappedMaterial = ReferencesToRemap.Find(MaterialsToRefreshAfterVirtualTextureConversion[MaterialIndex]))
-				{
-					MaterialsToRefreshAfterVirtualTextureConversion[MaterialIndex] = Cast<UMaterial>(*RemappedMaterial);
-				}
-			}
-
-			TArray<UTexture2D*> TexturesToConvertList(VirtualTexturesToConvert.Array());
-			IAssetTools& AssetTools = FModuleManager::LoadModuleChecked<FAssetToolsModule>("AssetTools").Get();
-			AssetTools.ConvertVirtualTextures(TexturesToConvertList, true, &MaterialsToRefreshAfterVirtualTextureConversion);
-		}
-	}
-
-	/**
-	 * Helper struct used to manage and validate the actor's state when modifying it during the finalization.
-	 */
-	struct FScopedFinalizeActorChanges
-	{
-		FDatasmithImportContext& ImportContext;
-		TSet<UActorComponent*> ComponentsToValidate;
-		AActor* FinalizedActor;
-
-		FScopedFinalizeActorChanges(AActor* InFinalizedActor, FDatasmithImportContext& InImportContext)
-			: ImportContext(InImportContext)
-			, FinalizedActor(InFinalizedActor)
-		{
-			// In order to allow modification on components owned by ExistingActor, unregister all of them
-			FinalizedActor->UnregisterAllComponents( /* bForReregister = */true);
-
-			//Some new components might be created when finalizing the actor, only validate those that we unregistered.
-			for (UActorComponent* Component : FinalizedActor->GetComponents())
-			{
-				ComponentsToValidate.Add(Component);
-			}
-		}
-
-		~FScopedFinalizeActorChanges()
-		{
-			for (UActorComponent* Component : FinalizedActor->GetComponents())
-			{
-				if (Component->IsRegistered() && ComponentsToValidate.Contains(Component))
-				{
-					ensureMsgf(false, TEXT("All components should still be unregistered at this point. Otherwise some datasmith templates might not have been applied properly."));
-					break;
-				}
-			}
-			
-			const FQuat PreviousRotation(FinalizedActor->GetRootComponent()->GetRelativeTransform().GetRotation());
-			FinalizedActor->PostEditChange();
-			FinalizedActor->RegisterAllComponents();
-
-			const bool bHasPostEditChangeModifiedRotation = !PreviousRotation.Equals(FinalizedActor->GetRootComponent()->GetRelativeTransform().GetRotation());
-			if (bHasPostEditChangeModifiedRotation)
-			{
-				//SingularityThreshold value is comming from the FQuat::Rotator() function, but is more permissive because the rotation is already diverging before the singularity threshold is reached.
-				const float SingularityThreshold = 0.4999f;
-				const float SingularityTest = PreviousRotation.Z * PreviousRotation.X - PreviousRotation.W * PreviousRotation.Y;
-				const AActor* RootSceneActor = ImportContext.ActorsContext.ImportSceneActor;
-
-				if (FinalizedActor != RootSceneActor
-					&& FMath::Abs(SingularityTest) > SingularityThreshold)
-				{
-					//This is a warning to explain the edge-case of UE-75467 while it's being fixed.
-					FFormatNamedArguments FormatArgs;
-					FormatArgs.Add(TEXT("ActorName"), FText::FromName(FinalizedActor->GetFName()));
-					ImportContext.LogWarning(FText::GetEmpty())
-						->AddToken(FUObjectToken::Create(FinalizedActor))
-						->AddToken(FTextToken::Create(FText::Format(LOCTEXT("UnsupportedRotationValueError", "The actor '{ActorName}' has a rotation value pointing to either (0, 90, 0) or (0, -90, 0)."
-							"This is an edge case that is not well supported in Unreal and can cause incorrect results."
-							"In those cases, it is recommended to bake the actor's transform into the mesh at export."), FormatArgs)));
-				}
-			}
-		}
-	};
-}
-
 void FDatasmithImporter::ImportStaticMeshes( FDatasmithImportContext& ImportContext )
 {
 	const int32 StaticMeshesCount = ImportContext.FilteredScene->GetMeshesCount();
@@ -985,7 +107,7 @@ void FDatasmithImporter::ImportStaticMeshes( FDatasmithImportContext& ImportCont
 	{
 		for (int32 MeshIndex = 0; MeshIndex < StaticMeshesCount && !ImportContext.bUserCancelled; ++MeshIndex)
 		{
-			ImportContext.bUserCancelled |= DatasmithImporterImpl::HasUserCancelledTask( ImportContext.FeedbackContext );
+			ImportContext.bUserCancelled |= FDatasmithImporterImpl::HasUserCancelledTask( ImportContext.FeedbackContext );
 
 
 			if (!ImportContext.AssetsContext.StaticMeshesFinalPackage || ImportContext.AssetsContext.StaticMeshesFinalPackage->GetFName() == NAME_None || ImportContext.SceneTranslator == nullptr)
@@ -1027,11 +149,11 @@ void FDatasmithImporter::ImportStaticMeshes( FDatasmithImportContext& ImportCont
 	// This pass will wait on the futures we got from the first pass async tasks
 	for ( int32 MeshIndex = 0; MeshIndex < StaticMeshesCount && !ImportContext.bUserCancelled; ++MeshIndex )
 	{
-		ImportContext.bUserCancelled |= DatasmithImporterImpl::HasUserCancelledTask( ImportContext.FeedbackContext );
+		ImportContext.bUserCancelled |= FDatasmithImporterImpl::HasUserCancelledTask( ImportContext.FeedbackContext );
 
 		TSharedRef< IDatasmithMeshElement > MeshElement = ImportContext.FilteredScene->GetMesh( MeshIndex ).ToSharedRef();
 
-		DatasmithImporterImpl::ReportProgress( Progress, 1.f, FText::FromString( FString::Printf( TEXT("Importing static mesh %d/%d (%s) ..."), MeshIndex + 1, StaticMeshesCount, MeshElement->GetLabel() ) ) );
+		FDatasmithImporterImpl::ReportProgress( Progress, 1.f, FText::FromString( FString::Printf( TEXT("Importing static mesh %d/%d (%s) ..."), MeshIndex + 1, StaticMeshesCount, MeshElement->GetLabel() ) ) );
 
 		UStaticMesh* ExistingStaticMesh = nullptr;
 
@@ -1142,9 +264,7 @@ UStaticMesh* FDatasmithImporter::ImportStaticMesh( FDatasmithImportContext& Impo
 
 UStaticMesh* FDatasmithImporter::FinalizeStaticMesh( UStaticMesh* SourceStaticMesh, const TCHAR* StaticMeshesFolderPath, UStaticMesh* ExistingStaticMesh, TMap< UObject*, UObject* >* ReferencesToRemap, bool bBuild)
 {
-	using namespace DatasmithImporterImpl;
-
-	UStaticMesh* DestinationStaticMesh = Cast< UStaticMesh >( FinalizeAsset( SourceStaticMesh, StaticMeshesFolderPath, ExistingStaticMesh, ReferencesToRemap ) );
+	UStaticMesh* DestinationStaticMesh = Cast< UStaticMesh >( FDatasmithImporterImpl::FinalizeAsset( SourceStaticMesh, StaticMeshesFolderPath, ExistingStaticMesh, ReferencesToRemap ) );
 
 	if (bBuild)
 	{
@@ -1180,7 +300,7 @@ void FDatasmithImporter::ImportTextures( FDatasmithImportContext& ImportContext
 {
 	TRACE_CPUPROFILER_EVENT_SCOPE(FDatasmithImporter::ImportTextures);
 
-	DatasmithImporterImpl::SetTexturesMode( ImportContext );
+	FDatasmithImporterImpl::SetTexturesMode( ImportContext );
 
 	const int32 TexturesCount = ImportContext.FilteredScene->GetTexturesCount();
 
@@ -1222,7 +342,7 @@ void FDatasmithImporter::ImportTextures( FDatasmithImportContext& ImportContext
 
 		for ( int32 TextureIndex = 0; TextureIndex < FilteredTextureElements.Num(); TextureIndex++ )
 		{
-			ImportContext.bUserCancelled |= DatasmithImporterImpl::HasUserCancelledTask( ImportContext.FeedbackContext );
+			ImportContext.bUserCancelled |= FDatasmithImporterImpl::HasUserCancelledTask( ImportContext.FeedbackContext );
 
 			AsyncData[TextureIndex].Result = 
 				Async(
@@ -1244,7 +364,7 @@ void FDatasmithImporter::ImportTextures( FDatasmithImportContext& ImportContext
 
 		for ( int32 TextureIndex = 0; TextureIndex < FilteredTextureElements.Num(); TextureIndex++ )
 		{
-			ImportContext.bUserCancelled |= DatasmithImporterImpl::HasUserCancelledTask( ImportContext.FeedbackContext );
+			ImportContext.bUserCancelled |= FDatasmithImporterImpl::HasUserCancelledTask( ImportContext.FeedbackContext );
 
 			if ( ImportContext.bUserCancelled )
 			{
@@ -1255,7 +375,7 @@ void FDatasmithImporter::ImportTextures( FDatasmithImportContext& ImportContext
 			{
 				const TSharedPtr< IDatasmithTextureElement >& TextureElement = FilteredTextureElements[TextureIndex];
 
-				DatasmithImporterImpl::ReportProgress( Progress, 1.f, FText::FromString( FString::Printf( TEXT("Importing texture %d/%d (%s) ..."), TextureIndex + 1, FilteredTextureElements.Num(), TextureElement->GetLabel() ) ) );
+				FDatasmithImporterImpl::ReportProgress( Progress, 1.f, FText::FromString( FString::Printf( TEXT("Importing texture %d/%d (%s) ..."), TextureIndex + 1, FilteredTextureElements.Num(), TextureElement->GetLabel() ) ) );
 
 				UTexture* ExistingTexture = nullptr;
 
@@ -1301,7 +421,7 @@ UTexture* FDatasmithImporter::ImportTexture( FDatasmithImportContext& ImportCont
 
 UTexture* FDatasmithImporter::FinalizeTexture( UTexture* SourceTexture, const TCHAR* TexturesFolderPath, UTexture* ExistingTexture, TMap< UObject*, UObject* >* ReferencesToRemap )
 {
-	return Cast< UTexture >( DatasmithImporterImpl::FinalizeAsset( SourceTexture, TexturesFolderPath, ExistingTexture, ReferencesToRemap ) );
+	return Cast< UTexture >( FDatasmithImporterImpl::FinalizeAsset( SourceTexture, TexturesFolderPath, ExistingTexture, ReferencesToRemap ) );
 }
 
 void FDatasmithImporter::ImportMaterials( FDatasmithImportContext& ImportContext )
@@ -1365,7 +485,7 @@ UMaterialFunction* FDatasmithImporter::ImportMaterialFunction(FDatasmithImportCo
 UMaterialFunction* FDatasmithImporter::FinalizeMaterialFunction(UObject* SourceMaterialFunction, const TCHAR* MaterialFunctionsFolderPath,
 	UMaterialFunction* ExistingMaterialFunction, TMap< UObject*, UObject* >* ReferencesToRemap)
 {
-	return Cast< UMaterialFunction >(DatasmithImporterImpl::FinalizeAsset(SourceMaterialFunction, MaterialFunctionsFolderPath, ExistingMaterialFunction, ReferencesToRemap));
+	return Cast< UMaterialFunction >(FDatasmithImporterImpl::FinalizeAsset(SourceMaterialFunction, MaterialFunctionsFolderPath, ExistingMaterialFunction, ReferencesToRemap));
 }
 
 UMaterialInterface* FDatasmithImporter::ImportMaterial( FDatasmithImportContext& ImportContext,
@@ -1420,8 +540,6 @@ UObject* FDatasmithImporter::FinalizeMaterial( UObject* SourceMaterial, const TC
 {
 	TRACE_CPUPROFILER_EVENT_SCOPE(FDatasmithImporter::FinalizeMaterial);
 
-	using namespace DatasmithImporterImpl;
-
 	// Finalizing the master material might add a remapping for the instance parent property so make sure we have a remapping map available
 	TOptional< TMap< UObject*, UObject* > > ReferencesToRemapLocal;
 	if ( !ReferencesToRemap )
@@ -1454,9 +572,9 @@ UObject* FDatasmithImporter::FinalizeMaterial( UObject* SourceMaterial, const TC
 
 	UMaterialEditingLibrary::DeleteAllMaterialExpressions( Cast< UMaterial >( ExistingMaterial ) );
 
-	UObject* DestinationMaterial = FinalizeAsset( SourceMaterial, MaterialFolderPath, ExistingMaterial, ReferencesToRemap );
+	UObject* DestinationMaterial = FDatasmithImporterImpl::FinalizeAsset( SourceMaterial, MaterialFolderPath, ExistingMaterial, ReferencesToRemap );
 
-	DatasmithImporterImpl::CompileMaterial( DestinationMaterial );
+	FDatasmithImporterImpl::CompileMaterial( DestinationMaterial );
 
 	return DestinationMaterial;
 }
@@ -1477,7 +595,7 @@ void FDatasmithImporter::ImportActors( FDatasmithImportContext& ImportContext )
 			continue;
 		}
 
-		if ( ImportContext.SceneAsset == SceneActor->Scene )
+		if ( ImportContext.SceneAsset == SceneActor->Scene && SceneActor->GetLevel() == ImportContext.ActorsContext.FinalWorld->GetCurrentLevel())
 		{
 			for ( TPair< FName, TSoftObjectPtr< AActor > >& Pair : SceneActor->RelatedActors )
 			{
@@ -1521,13 +639,13 @@ void FDatasmithImporter::ImportActors( FDatasmithImportContext& ImportContext )
 
 		for (int32 i = 0; i < ActorsCount && !ImportContext.bUserCancelled; ++i)
 		{
-			ImportContext.bUserCancelled |= DatasmithImporterImpl::HasUserCancelledTask( ImportContext.FeedbackContext );
+			ImportContext.bUserCancelled |= FDatasmithImporterImpl::HasUserCancelledTask( ImportContext.FeedbackContext );
 
 			TSharedPtr< IDatasmithActorElement > ActorElement = ImportContext.Scene->GetActor(i);
 
 			if ( ActorElement.IsValid() )
 			{
-				DatasmithImporterImpl::ReportProgress( Progress, 1.f, FText::FromString( FString::Printf( TEXT("Spawning actor %d/%d (%s) ..."), ( i + 1 ), ActorsCount, ActorElement->GetLabel() ) ) );
+				FDatasmithImporterImpl::ReportProgress( Progress, 1.f, FText::FromString( FString::Printf( TEXT("Spawning actor %d/%d (%s) ..."), ( i + 1 ), ActorsCount, ActorElement->GetLabel() ) ) );
 
 				if ( ActorElement->IsAComponent() )
 				{
@@ -1550,7 +668,7 @@ void FDatasmithImporter::ImportActors( FDatasmithImportContext& ImportContext )
 		// After all actors were imported, perform a post import step so that any dependencies can be resolved
 		for (int32 i = 0; i < ActorsCount && !ImportContext.bUserCancelled; ++i)
 		{
-			ImportContext.bUserCancelled |= DatasmithImporterImpl::HasUserCancelledTask( ImportContext.FeedbackContext );
+			ImportContext.bUserCancelled |= FDatasmithImporterImpl::HasUserCancelledTask( ImportContext.FeedbackContext );
 
 			TSharedPtr< IDatasmithActorElement > ActorElement = ImportContext.Scene->GetActor(i);
 
@@ -1573,7 +691,7 @@ void FDatasmithImporter::ImportActors( FDatasmithImportContext& ImportContext )
 
 	if ( ImportContext.bUserCancelled )
 	{
-		DatasmithImporterImpl::DeleteImportSceneActorIfNeeded(ImportContext.ActorsContext, true );
+		FDatasmithImporterImpl::DeleteImportSceneActorIfNeeded(ImportContext.ActorsContext, true );
 	}
 }
 
@@ -1636,7 +754,7 @@ AActor* FDatasmithImporter::ImportActor( FDatasmithImportContext& ImportContext,
 
 	for (int32 i = 0; i < ActorElement->GetChildrenCount() && !ImportContext.bUserCancelled; ++i)
 	{
-		ImportContext.bUserCancelled |= DatasmithImporterImpl::HasUserCancelledTask( ImportContext.FeedbackContext );
+		ImportContext.bUserCancelled |= FDatasmithImporterImpl::HasUserCancelledTask( ImportContext.FeedbackContext );
 
 		const TSharedPtr< IDatasmithActorElement >& ChildActorElement = ActorElement->GetChild(i);
 
@@ -1733,8 +851,6 @@ void FDatasmithImporter::FinalizeActors( FDatasmithImportContext& ImportContext,
 {
 	TRACE_CPUPROFILER_EVENT_SCOPE(FDatasmithImporter::FinalizeActors);
 
-	using namespace DatasmithImporterImpl;
-
 	if ( !ImportContext.bUserCancelled )
 	{
 		// Ensure a proper setup for the finalize of the actors
@@ -1792,7 +908,7 @@ void FDatasmithImporter::FinalizeActors( FDatasmithImportContext& ImportContext,
 				continue;
 			}
 
-			if ( ImportSceneActor->Scene != DestinationSceneActor->Scene )
+			if ( ImportSceneActor->Scene != DestinationSceneActor->Scene  || DestinationSceneActor->GetLevel() != ImportContext.ActorsContext.FinalWorld->GetCurrentLevel())
 			{
 				continue;
 			}
@@ -1850,7 +966,7 @@ void FDatasmithImporter::FinalizeActors( FDatasmithImportContext& ImportContext,
 			{
 				if ( AActor* Actor = DestinationActorPair.Value.Get() )
 				{
-					FixReferencesForObject( Actor, PerSceneActorReferencesToRemap );
+					FDatasmithImporterImpl::FixReferencesForObject( Actor, PerSceneActorReferencesToRemap );
 				}
 			}
 
@@ -1884,7 +1000,7 @@ void FDatasmithImporter::FinalizeActors( FDatasmithImportContext& ImportContext,
 		AssetToolsModule.Get().RenameReferencingSoftObjectPaths( PackagesToFix, RenamedActorsMap );
 	}
 
-	DeleteImportSceneActorIfNeeded( ImportContext.ActorsContext );
+	FDatasmithImporterImpl::DeleteImportSceneActorIfNeeded( ImportContext.ActorsContext );
 
 	// Ensure layer visibility is properly updated for new actors associated with existing layers
 	ULayersSubsystem* LayersSubsystem = GEditor->GetEditorSubsystem<ULayersSubsystem>();
@@ -1897,8 +1013,6 @@ AActor* FDatasmithImporter::FinalizeActor( FDatasmithImportContext& ImportContex
 {
 	TRACE_CPUPROFILER_EVENT_SCOPE(FDatasmithImporter::FinalizeActor);
 
-	using namespace DatasmithImporterImpl;
-
 	// If the existing actor is not of the same class we destroy it
 	if ( ExistingActor && ExistingActor->GetClass() != SourceActor.GetClass() )
 	{
@@ -1924,11 +1038,11 @@ AActor* FDatasmithImporter::FinalizeActor( FDatasmithImportContext& ImportContex
 	check( DestinationActor );
 	{
 		// Setup the actor to allow modifications.
-		FScopedFinalizeActorChanges ScopedFinalizedActorChanges(DestinationActor, ImportContext);
+		FDatasmithImporterImpl::FScopedFinalizeActorChanges ScopedFinalizedActorChanges(DestinationActor, ImportContext);
 		
 		ReferencesToRemap.Add( &SourceActor ) = DestinationActor;
 
-		TArray< FMigratedTemplatePairType > MigratedTemplates = MigrateTemplates(
+		TArray< FDatasmithImporterImpl::FMigratedTemplatePairType > MigratedTemplates = FDatasmithImporterImpl::MigrateTemplates(
 			&SourceActor,
 			ExistingActor,
 			&ReferencesToRemap,
@@ -1938,16 +1052,16 @@ AActor* FDatasmithImporter::FinalizeActor( FDatasmithImportContext& ImportContex
 		// Copy actor data
 		{
 			TArray< uint8 > Bytes;
-			FActorWriter ObjectWriter( &SourceActor, Bytes );
+			FDatasmithImporterImpl::FActorWriter ObjectWriter( &SourceActor, Bytes );
 			FObjectReader ObjectReader( DestinationActor, Bytes );
 		}
 
-		FixReferencesForObject( DestinationActor, ReferencesToRemap );
+		FDatasmithImporterImpl::FixReferencesForObject( DestinationActor, ReferencesToRemap );
 
-		DatasmithImporterImpl::FinalizeComponents( ImportContext, SourceActor, *DestinationActor, ReferencesToRemap );
+		FDatasmithImporterImpl::FinalizeComponents( ImportContext, SourceActor, *DestinationActor, ReferencesToRemap );
 
 		// The templates for the actor need to be applied after the components were created.
-		ApplyMigratedTemplates( MigratedTemplates, DestinationActor );
+		FDatasmithImporterImpl::ApplyMigratedTemplates( MigratedTemplates, DestinationActor );
 
 		// Restore hierarchy
 		for ( AActor* Child : Children )
@@ -1995,7 +1109,7 @@ void FDatasmithImporter::ImportLevelSequences( FDatasmithImportContext& ImportCo
 	SequencesToImport.Reserve(SequencesCount);
 	for ( int32 SequenceIndex = 0; SequenceIndex < SequencesCount && !ImportContext.bUserCancelled; ++SequenceIndex )
 	{
-		ImportContext.bUserCancelled |= DatasmithImporterImpl::HasUserCancelledTask( ImportContext.FeedbackContext );
+		ImportContext.bUserCancelled |= FDatasmithImporterImpl::HasUserCancelledTask( ImportContext.FeedbackContext );
 
 		TSharedPtr< IDatasmithLevelSequenceElement > SequenceElement = ImportContext.FilteredScene->GetLevelSequence( SequenceIndex );
 		if ( !SequenceElement )
@@ -2018,7 +1132,7 @@ void FDatasmithImporter::ImportLevelSequences( FDatasmithImportContext& ImportCo
 		// Scan remaining sequences and import the ones we can, removing from this array
 		for ( int32 SequenceIndex = SequencesToImport.Num() - 1; SequenceIndex >= 0 && !ImportContext.bUserCancelled; --SequenceIndex )
 		{
-			ImportContext.bUserCancelled |= DatasmithImporterImpl::HasUserCancelledTask( ImportContext.FeedbackContext );
+			ImportContext.bUserCancelled |= FDatasmithImporterImpl::HasUserCancelledTask( ImportContext.FeedbackContext );
 
 			TSharedPtr<IDatasmithLevelSequenceElement>& SequenceElement = SequencesToImport[SequenceIndex];
 
@@ -2039,7 +1153,7 @@ void FDatasmithImporter::ImportLevelSequences( FDatasmithImportContext& ImportCo
 		}
 
 		FString SequenceName = ObjectTools::SanitizeObjectName(SequenceElement->GetName());
-		DatasmithImporterImpl::ReportProgress( Progress, 1.f, FText::FromString( FString::Printf(TEXT("Importing level sequence %d/%d (%s) ..."), NumImported + 1, HardLoopCounter, *SequenceName ) ) );
+		FDatasmithImporterImpl::ReportProgress( Progress, 1.f, FText::FromString( FString::Printf(TEXT("Importing level sequence %d/%d (%s) ..."), NumImported + 1, HardLoopCounter, *SequenceName ) ) );
 
 		ULevelSequence*& ImportedLevelSequence = ImportContext.ImportedLevelSequences.FindOrAdd( SequenceElement.ToSharedRef() );
 		if (ImportContext.SceneTranslator)
@@ -2078,7 +1192,7 @@ void FDatasmithImporter::ImportLevelSequences( FDatasmithImportContext& ImportCo
 
 ULevelSequence* FDatasmithImporter::FinalizeLevelSequence( ULevelSequence* SourceLevelSequence, const TCHAR* AnimationsFolderPath, ULevelSequence* ExistingLevelSequence )
 {
-	return Cast< ULevelSequence >( DatasmithImporterImpl::PublicizeAsset( SourceLevelSequence, AnimationsFolderPath, ExistingLevelSequence ) );
+	return Cast< ULevelSequence >( FDatasmithImporterImpl::PublicizeAsset( SourceLevelSequence, AnimationsFolderPath, ExistingLevelSequence ) );
 }
 
 void FDatasmithImporter::ImportLevelVariantSets( FDatasmithImportContext& ImportContext )
@@ -2102,7 +1216,7 @@ void FDatasmithImporter::ImportLevelVariantSets( FDatasmithImportContext& Import
 
 	for ( int32 LevelVariantSetIndex = 0; LevelVariantSetIndex < LevelVariantSetsCount && !ImportContext.bUserCancelled; ++LevelVariantSetIndex )
 	{
-		ImportContext.bUserCancelled |= DatasmithImporterImpl::HasUserCancelledTask( ImportContext.FeedbackContext );
+		ImportContext.bUserCancelled |= FDatasmithImporterImpl::HasUserCancelledTask( ImportContext.FeedbackContext );
 
 		TSharedPtr< IDatasmithLevelVariantSetsElement > LevelVariantSetsElement = ImportContext.FilteredScene->GetLevelVariantSets( LevelVariantSetIndex );
 		if ( !LevelVariantSetsElement )
@@ -2123,7 +1237,7 @@ void FDatasmithImporter::ImportLevelVariantSets( FDatasmithImportContext& Import
 
 
 		FString LevelVariantSetsName = ObjectTools::SanitizeObjectName(LevelVariantSetsElement->GetName());
-		DatasmithImporterImpl::ReportProgress(Progress, 1.f, FText::FromString(FString::Printf(TEXT("Importing level variant sets %d/%d (%s) ..."), LevelVariantSetIndex + 1, LevelVariantSetsCount, *LevelVariantSetsName)));
+		FDatasmithImporterImpl::ReportProgress(Progress, 1.f, FText::FromString(FString::Printf(TEXT("Importing level variant sets %d/%d (%s) ..."), LevelVariantSetIndex + 1, LevelVariantSetsCount, *LevelVariantSetsName)));
 
 		ULevelVariantSets*& ImportedLevelVariantSets = ImportContext.ImportedLevelVariantSets.FindOrAdd( LevelVariantSetsElement.ToSharedRef() );
 		ImportedLevelVariantSets = FDatasmithLevelVariantSetsImporter::ImportLevelVariantSets( LevelVariantSetsElement.ToSharedRef(), ImportContext, ExistingLevelVariantSets );
@@ -2135,7 +1249,7 @@ void FDatasmithImporter::ImportLevelVariantSets( FDatasmithImportContext& Import
 
 ULevelVariantSets* FDatasmithImporter::FinalizeLevelVariantSets( ULevelVariantSets* SourceLevelVariantSets, const TCHAR* VariantsFolderPath, ULevelVariantSets* ExistingLevelVariantSets )
 {
-	return Cast< ULevelVariantSets >( DatasmithImporterImpl::PublicizeAsset( SourceLevelVariantSets, VariantsFolderPath, ExistingLevelVariantSets ) );
+	return Cast< ULevelVariantSets >( FDatasmithImporterImpl::PublicizeAsset( SourceLevelVariantSets, VariantsFolderPath, ExistingLevelVariantSets ) );
 }
 
 void FDatasmithImporter::ImportMetaDataForObject(FDatasmithImportContext& ImportContext, const TSharedRef<IDatasmithElement>& DatasmithElement, UObject* Object)
@@ -2352,7 +1466,7 @@ void FDatasmithImporter::FinalizeImport(FDatasmithImportContext& ImportContext,
 
 		FName TextureId = ImportedTexturePair.Key->GetName();
 
-		DatasmithImporterImpl::ReportProgress(Progress, 1.f, FText::FromString(FString::Printf(TEXT("Finalizing assets %d/%d (Texture %s) ..."), ++AssetIndex, NumAssetsToFinalize, *SourceTexture->GetName())));
+		FDatasmithImporterImpl::ReportProgress(Progress, 1.f, FText::FromString(FString::Printf(TEXT("Finalizing assets %d/%d (Texture %s) ..."), ++AssetIndex, NumAssetsToFinalize, *SourceTexture->GetName())));
 
 		TSoftObjectPtr< UTexture >& ExistingTexturePtr = ImportContext.SceneAsset->Textures.FindOrAdd(TextureId);
 		UTexture* ExistingTexture = ExistingTexturePtr.Get();
@@ -2361,7 +1475,7 @@ void FDatasmithImporter::FinalizeImport(FDatasmithImportContext& ImportContext,
 		FString DestinationPackagePath = SourcePackagePath.Replace(*TransientFolderPath, *RootFolderPath, ESearchCase::CaseSensitive);
 
 		ExistingTexturePtr = FinalizeTexture(SourceTexture, *DestinationPackagePath, ExistingTexture, &ReferencesToRemap);
-		DatasmithImporterImpl::CheckAssetPersistenceValidity(ExistingTexturePtr.Get(), ImportContext);
+		FDatasmithImporterImpl::CheckAssetPersistenceValidity(ExistingTexturePtr.Get(), ImportContext);
 	}
 
 	// Unregister all actors component to avoid excessive refresh in the 3D engine while updating materials.
@@ -2388,7 +1502,7 @@ void FDatasmithImporter::FinalizeImport(FDatasmithImportContext& ImportContext,
 		}
 
 		FName MaterialFunctionId = ImportedMaterialFunctionPair.Key->GetName();
-		DatasmithImporterImpl::ReportProgress(Progress, 1.f, FText::FromString(FString::Printf(TEXT("Finalizing assets %d/%d (Material Function %s) ..."), ++AssetIndex, NumAssetsToFinalize, *SourceMaterialFunction->GetName())));
+		FDatasmithImporterImpl::ReportProgress(Progress, 1.f, FText::FromString(FString::Printf(TEXT("Finalizing assets %d/%d (Material Function %s) ..."), ++AssetIndex, NumAssetsToFinalize, *SourceMaterialFunction->GetName())));
 
 		TSoftObjectPtr< UMaterialFunction >& ExistingMaterialFunctionPtr = ImportContext.SceneAsset->MaterialFunctions.FindOrAdd(MaterialFunctionId);
 		UMaterialFunction* ExistingMaterialFunction = ExistingMaterialFunctionPtr.Get();
@@ -2397,7 +1511,7 @@ void FDatasmithImporter::FinalizeImport(FDatasmithImportContext& ImportContext,
 		FString DestinationPackagePath = SourcePackagePath.Replace(*TransientFolderPath, *RootFolderPath, ESearchCase::CaseSensitive);
 
 		ExistingMaterialFunctionPtr = FinalizeMaterialFunction(SourceMaterialFunction, *DestinationPackagePath, ExistingMaterialFunction, &ReferencesToRemap);
-		DatasmithImporterImpl::CheckAssetPersistenceValidity(ExistingMaterialFunctionPtr.Get(), ImportContext);
+		FDatasmithImporterImpl::CheckAssetPersistenceValidity(ExistingMaterialFunctionPtr.Get(), ImportContext);
 	}
 
 	TArray<UMaterial*> MaterialsToRefreshAfterVirtualTextureConversion;
@@ -2417,7 +1531,7 @@ void FDatasmithImporter::FinalizeImport(FDatasmithImportContext& ImportContext,
 
 		FName MaterialId = ImportedMaterialPair.Key->GetName();
 
-		DatasmithImporterImpl::ReportProgress(Progress, 1.f, FText::FromString(FString::Printf(TEXT("Finalizing assets %d/%d (Material %s) ..."), ++AssetIndex, NumAssetsToFinalize, *SourceMaterialInterface->GetName())));
+		FDatasmithImporterImpl::ReportProgress(Progress, 1.f, FText::FromString(FString::Printf(TEXT("Finalizing assets %d/%d (Material %s) ..."), ++AssetIndex, NumAssetsToFinalize, *SourceMaterialInterface->GetName())));
 
 		TSoftObjectPtr< UMaterialInterface >& ExistingMaterialPtr = ImportContext.SceneAsset->Materials.FindOrAdd(MaterialId);
 		UMaterialInterface* ExistingMaterial = ExistingMaterialPtr.Get();
@@ -2455,11 +1569,11 @@ void FDatasmithImporter::FinalizeImport(FDatasmithImportContext& ImportContext,
 		if (ExistingMaterialPtr.IsValid())
 		{
 			PackagesToCheck.Add(ExistingMaterialPtr->GetOutermost());
-			DatasmithImporterImpl::CheckAssetPersistenceValidity(ExistingMaterialPtr.Get(), ImportContext);
+			FDatasmithImporterImpl::CheckAssetPersistenceValidity(ExistingMaterialPtr.Get(), ImportContext);
 		}
 	}
 
-	DatasmithImporterImpl::ConvertUnsupportedVirtualTexture(ImportContext, ImportContext.AssetsContext.VirtualTexturesToConvert, ReferencesToRemap);
+	FDatasmithImporterImpl::ConvertUnsupportedVirtualTexture(ImportContext, ImportContext.AssetsContext.VirtualTexturesToConvert, ReferencesToRemap);
 
 	// Materials have been updated, we can register everything back.
 	for (TObjectIterator<AActor> ActorIterator; ActorIterator; ++ActorIterator)
@@ -2488,7 +1602,7 @@ void FDatasmithImporter::FinalizeImport(FDatasmithImportContext& ImportContext,
 
 		FName StaticMeshId = ImportedStaticMeshPair.Key->GetName();
 
-		DatasmithImporterImpl::ReportProgress(Progress, 1.f, FText::FromString(FString::Printf(TEXT("Finalizing assets %d/%d (Static Mesh %s) ..."), ++AssetIndex, NumAssetsToFinalize, *SourceStaticMesh->GetName())));
+		FDatasmithImporterImpl::ReportProgress(Progress, 1.f, FText::FromString(FString::Printf(TEXT("Finalizing assets %d/%d (Static Mesh %s) ..."), ++AssetIndex, NumAssetsToFinalize, *SourceStaticMesh->GetName())));
 
 		TSoftObjectPtr< UStaticMesh >& ExistingStaticMeshPtr = ImportContext.SceneAsset->StaticMeshes.FindOrAdd(StaticMeshId);
 		UStaticMesh* ExistingStaticMesh = ExistingStaticMeshPtr.Get();
@@ -2497,7 +1611,7 @@ void FDatasmithImporter::FinalizeImport(FDatasmithImportContext& ImportContext,
 		FString DestinationPackagePath = SourcePackagePath.Replace( *TransientFolderPath, *RootFolderPath, ESearchCase::CaseSensitive );
 
 		ExistingStaticMeshPtr = FinalizeStaticMesh(SourceStaticMesh, *DestinationPackagePath, ExistingStaticMesh, &ReferencesToRemap, false);
-		DatasmithImporterImpl::CheckAssetPersistenceValidity(ExistingStaticMeshPtr.Get(), ImportContext);
+		FDatasmithImporterImpl::CheckAssetPersistenceValidity(ExistingStaticMeshPtr.Get(), ImportContext);
 
 		ImportedStaticMeshPair.Value = ExistingStaticMeshPtr.Get();
 		StaticMeshes.Add(ExistingStaticMeshPtr.Get());
@@ -2506,7 +1620,7 @@ void FDatasmithImporter::FinalizeImport(FDatasmithImportContext& ImportContext,
 	int32 StaticMeshIndex = 0;
 	auto ProgressFunction = [&](UStaticMesh* StaticMesh)
 	{
-		DatasmithImporterImpl::ReportProgress(Progress, 1.f, FText::FromString(FString::Printf(TEXT("Building Static Mesh %d/%d (%s) ..."), ++StaticMeshIndex, StaticMeshes.Num(), *StaticMesh->GetName())));
+		FDatasmithImporterImpl::ReportProgress(Progress, 1.f, FText::FromString(FString::Printf(TEXT("Building Static Mesh %d/%d (%s) ..."), ++StaticMeshIndex, StaticMeshes.Num(), *StaticMesh->GetName())));
 		return !ImportContext.bUserCancelled;
 	};
 
@@ -2528,7 +1642,7 @@ void FDatasmithImporter::FinalizeImport(FDatasmithImportContext& ImportContext,
 
 		FName LevelSequenceId = ImportedLevelSequencePair.Key->GetName();
 
-		DatasmithImporterImpl::ReportProgress(Progress, 1.f, FText::FromString(FString::Printf(TEXT("Finalizing assets %d/%d (Level Sequence %s) ..."), ++AssetIndex, NumAssetsToFinalize, *SourceLevelSequence->GetName())));
+		FDatasmithImporterImpl::ReportProgress(Progress, 1.f, FText::FromString(FString::Printf(TEXT("Finalizing assets %d/%d (Level Sequence %s) ..."), ++AssetIndex, NumAssetsToFinalize, *SourceLevelSequence->GetName())));
 
 
 		TSoftObjectPtr< ULevelSequence >& ExistingLevelSequencePtr = ImportContext.SceneAsset->LevelSequences.FindOrAdd(LevelSequenceId);
@@ -2557,7 +1671,7 @@ void FDatasmithImporter::FinalizeImport(FDatasmithImportContext& ImportContext,
 
 		FName LevelVariantSetsId = ImportedLevelVariantSetsPair.Key->GetName();
 
-		DatasmithImporterImpl::ReportProgress(Progress, 1.f, FText::FromString(FString::Printf(TEXT("Finalizing assets %d/%d (Level Variant Sets %s) ..."), ++AssetIndex, NumAssetsToFinalize, *SourceLevelVariantSets->GetName())));
+		FDatasmithImporterImpl::ReportProgress(Progress, 1.f, FText::FromString(FString::Printf(TEXT("Finalizing assets %d/%d (Level Variant Sets %s) ..."), ++AssetIndex, NumAssetsToFinalize, *SourceLevelVariantSets->GetName())));
 
 		TSoftObjectPtr< ULevelVariantSets >& ExistingLevelVariantSetsPtr = ImportContext.SceneAsset->LevelVariantSets.FindOrAdd(LevelVariantSetsId);
 		ULevelVariantSets* ExistingLevelVariantSets = ExistingLevelVariantSetsPtr.Get();
diff --git a/Engine/Plugins/Enterprise/DatasmithImporter/Source/DatasmithImporter/Private/Utility/DatasmithImporterImpl.cpp b/Engine/Plugins/Enterprise/DatasmithImporter/Source/DatasmithImporter/Private/Utility/DatasmithImporterImpl.cpp
new file mode 100644
--- /dev/null
+++ b/Engine/Plugins/Enterprise/DatasmithImporter/Source/DatasmithImporter/Private/Utility/DatasmithImporterImpl.cpp
@@ -0,0 +1,939 @@
+// Copyright Epic Games, Inc. All Rights Reserved.
+
+#include "DatasmithImporterImpl.h"
+
+#include "DatasmithActorImporter.h"
+#include "DatasmithAdditionalData.h"
+#include "DatasmithAssetImportData.h"
+#include "DatasmithAssetUserData.h"
+#include "DatasmithCameraImporter.h"
+#include "DatasmithImportContext.h"
+#include "DatasmithLevelSequenceImporter.h"
+#include "DatasmithLevelVariantSetsImporter.h"
+#include "DatasmithLightImporter.h"
+#include "DatasmithMaterialImporter.h"
+#include "DatasmithPayload.h"
+#include "DatasmithPostProcessImporter.h"
+#include "DatasmithScene.h"
+#include "DatasmithSceneActor.h"
+#include "DatasmithSceneFactory.h"
+#include "DatasmithStaticMeshImporter.h"
+#include "DatasmithTranslator.h"
+#include "DatasmithTextureImporter.h"
+#include "IDatasmithSceneElements.h"
+#include "DatasmithAnimationElements.h"
+#include "LevelVariantSets.h"
+#include "ObjectTemplates/DatasmithObjectTemplate.h"
+#include "Utility/DatasmithImporterUtils.h"
+#include "Utility/DatasmithTextureResize.h"
+
+#include "AssetRegistryModule.h"
+#include "AssetToolsModule.h"
+#include "Async/Async.h"
+#include "Async/AsyncWork.h"
+#include "CineCameraComponent.h"
+#include "ComponentReregisterContext.h"
+#include "Components/HierarchicalInstancedStaticMeshComponent.h"
+#include "Containers/Map.h"
+#include "CoreMinimal.h"
+#include "Editor/UnrealEdEngine.h"
+#include "EditorLevelUtils.h"
+#include "Engine/Engine.h"
+#include "Engine/StaticMesh.h"
+#include "Engine/StaticMeshActor.h"
+#include "HAL/FileManager.h"
+#include "IAssetTools.h"
+#include "ISourceControlModule.h"
+#include "ISourceControlProvider.h"
+#include "Kismet2/KismetEditorUtilities.h"
+#include "Landscape.h"
+#include "Layers/LayersSubsystem.h"
+#include "LevelSequence.h"
+#include "MaterialEditingLibrary.h"
+#include "MaterialShared.h"
+#include "Materials/Material.h"
+#include "Materials/MaterialFunction.h"
+#include "Materials/MaterialInstanceConstant.h"
+#include "Misc/FeedbackContext.h"
+#include "Misc/FileHelper.h"
+#include "Misc/ScopedSlowTask.h"
+#include "Misc/UObjectToken.h"
+#include "ObjectTools.h"
+#include "PackageTools.h"
+#include "Serialization/ArchiveReplaceObjectRef.h"
+#include "Serialization/MemoryReader.h"
+#include "Serialization/ObjectReader.h"
+#include "Settings/EditorExperimentalSettings.h"
+#include "SourceControlOperations.h"
+#include "Templates/UniquePtr.h"
+#include "UObject/Package.h"
+#include "UnrealEdGlobals.h"
+
+extern UNREALED_API UEditorEngine* GEditor;
+
+#define LOCTEXT_NAMESPACE "DatasmithImporter"
+
+void FDatasmithImporterImpl::ReportProgress(FScopedSlowTask* SlowTask, const float ExpectedWorkThisFrame, const FText& Text)
+{
+	if (SlowTask)
+	{
+		SlowTask->EnterProgressFrame(ExpectedWorkThisFrame, Text);
+	}
+}
+
+void FDatasmithImporterImpl::ReportProgress(FScopedSlowTask* SlowTask, const float ExpectedWorkThisFrame, FText&& Text)
+{
+	if (SlowTask)
+	{
+		SlowTask->EnterProgressFrame(ExpectedWorkThisFrame, Text);
+	}
+}
+
+bool FDatasmithImporterImpl::HasUserCancelledTask(FFeedbackContext* FeedbackContext)
+{
+	if (FeedbackContext)
+	{
+		return FeedbackContext->ReceivedUserCancel();
+	}
+
+	return false;
+}
+
+UObject* FDatasmithImporterImpl::PublicizeAsset( UObject* SourceAsset, const TCHAR* DestinationPath, UObject* ExistingAsset )
+{
+	TRACE_CPUPROFILER_EVENT_SCOPE(FDatasmithImporterImpl::PublicizeAsset);
+
+	UPackage* DestinationPackage;
+
+	if ( !ExistingAsset )
+	{
+		const FString AssetName = SourceAsset->GetName();
+		bool bPathIsComplete = AssetName == FPaths::GetBaseFilename( DestinationPath );
+		FString DestinationPackagePath = UPackageTools::SanitizePackageName( bPathIsComplete ? DestinationPath : FPaths::Combine( DestinationPath, AssetName ) );
+		FString DestinationAssetPath = DestinationPackagePath + TEXT(".") + UPackageTools::SanitizePackageName( AssetName );
+
+		ExistingAsset = FDatasmithImporterUtils::FindObject<UObject>( nullptr, DestinationAssetPath );
+
+		DestinationPackage = ExistingAsset ? ExistingAsset->GetOutermost() : CreatePackage( nullptr, *DestinationPackagePath );
+	}
+	else
+	{
+		DestinationPackage = ExistingAsset->GetOutermost();
+	}
+
+	DestinationPackage->FullyLoad();
+
+	UObject* DestinationAsset = ExistingAsset;
+
+	FString OldAssetPathName;
+
+	// If the object already exist, then we need to fix up the reference
+	if ( ExistingAsset != nullptr && ExistingAsset != SourceAsset )
+	{
+		OldAssetPathName = ExistingAsset->GetPathName();
+
+		DestinationAsset = FDatasmithImporterUtils::DuplicateObject( SourceAsset, DestinationPackage, ExistingAsset->GetFName() );
+
+		// If mesh's label has changed, update its name
+		if ( ExistingAsset->GetFName() != SourceAsset->GetFName() )
+		{
+			DestinationAsset->Rename( *SourceAsset->GetName(), DestinationPackage, REN_DontCreateRedirectors | REN_NonTransactional );
+		}
+
+		if ( UStaticMesh* DestinationMesh = Cast< UStaticMesh >( DestinationAsset ) )
+		{
+			// This is done during the mesh build process but we need to redo it after the DuplicateObject since the links are now valid
+			for ( TObjectIterator< UStaticMeshComponent > It; It; ++It )
+			{
+				if ( It->GetStaticMesh() == DestinationMesh )
+				{
+					It->FixupOverrideColorsIfNecessary( true );
+					It->InvalidateLightingCache();
+				}
+			}
+		}
+	}
+	else
+	{
+		SourceAsset->Rename( *SourceAsset->GetName(), DestinationPackage, REN_DontCreateRedirectors | REN_NonTransactional );
+		DestinationAsset = SourceAsset;
+	}
+
+	DestinationAsset->SetFlags( RF_Public );
+	DestinationAsset->MarkPackageDirty();
+
+	if ( !ExistingAsset )
+	{
+		FAssetRegistryModule::AssetCreated( DestinationAsset );
+	}
+	else if ( !OldAssetPathName.IsEmpty() )
+	{
+		FAssetRegistryModule::AssetRenamed( DestinationAsset, OldAssetPathName );
+	}
+
+	return DestinationAsset;
+}
+
+void FDatasmithImporterImpl::CheckAssetPersistenceValidity(UObject* Asset, FDatasmithImportContext& ImportContext)
+{
+	if (Asset == nullptr)
+	{
+		return;
+	}
+
+	UPackage* Package = Asset->GetOutermost();
+	const FString PackageName = Package->GetName();
+
+	CheckAssetPersistenceValidity(PackageName, ImportContext);
+}
+
+/** Set the texture mode on each texture element based on its usage in the materials */
+void FDatasmithImporterImpl::SetTexturesMode( FDatasmithImportContext& ImportContext )
+{
+	const int32 TexturesCount = ImportContext.FilteredScene->GetTexturesCount();
+	const int32 MaterialsCount = ImportContext.FilteredScene->GetMaterialsCount();
+
+	FFeedbackContext* FeedbackContext = ImportContext.FeedbackContext;
+	for ( int32 TextureIndex = 0; TextureIndex < TexturesCount && !ImportContext.bUserCancelled; ++TextureIndex )
+	{
+		ImportContext.bUserCancelled |= FDatasmithImporterImpl::HasUserCancelledTask( ImportContext.FeedbackContext );
+
+		TSharedPtr< IDatasmithTextureElement > TextureElement = ImportContext.FilteredScene->GetTexture( TextureIndex );
+		const FString TextureName = ObjectTools::SanitizeObjectName( TextureElement->GetName() );
+
+		for ( int32 MaterialIndex = 0; MaterialIndex < MaterialsCount; ++MaterialIndex )
+		{
+			const TSharedPtr< IDatasmithBaseMaterialElement >& BaseMaterialElement = ImportContext.FilteredScene->GetMaterial( MaterialIndex );
+
+			if ( BaseMaterialElement->IsA( EDatasmithElementType::Material ) )
+			{
+				const TSharedPtr< IDatasmithMaterialElement >& MaterialElement = StaticCastSharedPtr< IDatasmithMaterialElement >( BaseMaterialElement );
+
+				for ( int32 s = 0; s < MaterialElement->GetShadersCount(); ++s )
+				{
+					const TSharedPtr< IDatasmithShaderElement >& ShaderElement = MaterialElement->GetShader(s);
+
+					if ( FCString::Strlen( ShaderElement->GetDiffuseTexture() ) > 0 && ShaderElement->GetDiffuseTexture() == TextureName)
+					{
+						TextureElement->SetTextureMode(EDatasmithTextureMode::Diffuse);
+					}
+					else if ( FCString::Strlen( ShaderElement->GetReflectanceTexture() ) > 0 && ShaderElement->GetReflectanceTexture() == TextureName)
+					{
+						TextureElement->SetTextureMode(EDatasmithTextureMode::Specular);
+					}
+					else if ( FCString::Strlen( ShaderElement->GetDisplaceTexture() ) > 0 && ShaderElement->GetDisplaceTexture() == TextureName)
+					{
+						TextureElement->SetTextureMode(EDatasmithTextureMode::Displace);
+					}
+					else if ( FCString::Strlen( ShaderElement->GetNormalTexture() ) > 0 && ShaderElement->GetNormalTexture() == TextureName)
+					{
+						if (!ShaderElement->GetNormalTextureSampler().bInvert)
+						{
+							TextureElement->SetTextureMode(EDatasmithTextureMode::Normal);
+						}
+						else
+						{
+							TextureElement->SetTextureMode(EDatasmithTextureMode::NormalGreenInv);
+						}
+					}
+				}
+			}
+			else if ( BaseMaterialElement->IsA( EDatasmithElementType::UEPbrMaterial ) )
+			{
+				const TSharedPtr< IDatasmithUEPbrMaterialElement >& MaterialElement = StaticCastSharedPtr< IDatasmithUEPbrMaterialElement >( BaseMaterialElement );
+
+				TFunction< bool( IDatasmithMaterialExpression* ) > IsTextureConnected;
+				IsTextureConnected = [ &TextureName, &IsTextureConnected ]( IDatasmithMaterialExpression* MaterialExpression ) -> bool
+				{
+					if ( !MaterialExpression )
+					{
+						return false;
+					}
+
+					if ( MaterialExpression->IsA( EDatasmithMaterialExpressionType::Texture ) )
+					{
+						IDatasmithMaterialExpressionTexture* TextureExpression = static_cast< IDatasmithMaterialExpressionTexture* >( MaterialExpression );
+
+						if ( TextureExpression->GetTexturePathName() == TextureName )
+						{
+							return true;
+						}
+					}
+
+					for ( int32 InputIndex = 0; InputIndex < MaterialExpression->GetInputCount(); ++InputIndex )
+					{
+						IDatasmithMaterialExpression* ConnectedExpression = MaterialExpression->GetInput( InputIndex )->GetExpression();
+
+						if ( ConnectedExpression && IsTextureConnected( ConnectedExpression ) )
+						{
+							return true;
+						}
+					}
+
+					return false;
+				};
+
+				if ( IsTextureConnected( MaterialElement->GetBaseColor().GetExpression() ) )
+				{
+					TextureElement->SetTextureMode(EDatasmithTextureMode::Diffuse);
+				}
+				else if ( IsTextureConnected( MaterialElement->GetSpecular().GetExpression() ) )
+				{
+					TextureElement->SetTextureMode(EDatasmithTextureMode::Specular);
+				}
+				else if ( IsTextureConnected( MaterialElement->GetNormal().GetExpression() ) )
+				{
+					if ( TextureElement->GetTextureMode() != EDatasmithTextureMode::Bump )
+					{
+						TextureElement->SetTextureMode(EDatasmithTextureMode::Normal);
+					}
+				}
+			}
+		}
+	}
+}
+
+void FDatasmithImporterImpl::CompileMaterial( UObject* Material )
+{
+	if ( !Material->IsA< UMaterialInterface >() && !Material->IsA< UMaterialFunctionInterface >() )
+	{
+		return;
+	}
+
+	TRACE_CPUPROFILER_EVENT_SCOPE(FDatasmithImporterImpl::CompileMaterial);
+
+	FMaterialUpdateContext MaterialUpdateContext;
+
+	if ( UMaterialInterface* MaterialInterface = Cast< UMaterialInterface >( Material ) )
+	{
+		MaterialUpdateContext.AddMaterialInterface( MaterialInterface );
+	}
+
+	if ( UMaterialInstanceConstant* ConstantMaterialInstance = Cast< UMaterialInstanceConstant >( Material ) )
+	{
+		// If BlendMode override property has been changed, make sure this combination of the parent material is compiled
+		if ( ConstantMaterialInstance->BasePropertyOverrides.bOverride_BlendMode == true )
+		{
+			ConstantMaterialInstance->ForceRecompileForRendering();
+		}
+		else
+		{
+			// If a switch is overriden, we need to recompile
+			FStaticParameterSet StaticParameters;
+			ConstantMaterialInstance->GetStaticParameterValues( StaticParameters );
+
+			for ( FStaticSwitchParameter& Switch : StaticParameters.StaticSwitchParameters )
+			{
+				if ( Switch.bOverride )
+				{
+					ConstantMaterialInstance->ForceRecompileForRendering();
+					break;
+				}
+			}
+		}
+	}
+
+	Material->PreEditChange( nullptr );
+	Material->PostEditChange();
+}
+
+void FDatasmithImporterImpl::FixReferencesForObject( UObject* Object, const TMap< UObject*, UObject* >& ReferencesToRemap )
+{
+	constexpr bool bNullPrivateRefs = false;
+	constexpr bool bIgnoreOuterRef = true;
+	constexpr bool bIgnoreArchetypeRef = true;
+
+	if ( ReferencesToRemap.Num() > 0 )
+	{
+		FArchiveReplaceObjectRef< UObject > ArchiveReplaceObjectRef( Object, ReferencesToRemap, bNullPrivateRefs, bIgnoreOuterRef, bIgnoreArchetypeRef );
+	}
+}
+
+/**
+	* Creates templates to apply the values from the SourceObject on the DestinationObject.
+	*
+	* @returns An array of template pairs. The key is the template for the object, the value is a template to force apply to the object,
+	*			it contains the values from the key and any overrides that were present on the DestinationObject.
+	*/
+TArray< FDatasmithImporterImpl::FMigratedTemplatePairType > FDatasmithImporterImpl::MigrateTemplates( UObject* SourceObject, UObject* DestinationObject, const TMap< UObject*, UObject* >* ReferencesToRemap, bool bIsForActor )
+{
+	TArray< FMigratedTemplatePairType > Results;
+
+	if ( !SourceObject )
+	{
+		return Results;
+	}
+
+	TMap< TSubclassOf< UDatasmithObjectTemplate >, UDatasmithObjectTemplate* >* SourceTemplates = FDatasmithObjectTemplateUtils::FindOrCreateObjectTemplates( SourceObject );
+
+	if ( !SourceTemplates )
+	{
+		return Results;
+	}
+
+	for ( const TPair< TSubclassOf< UDatasmithObjectTemplate >, UDatasmithObjectTemplate* >& SourceTemplatePair : *SourceTemplates )
+	{
+		if ( bIsForActor == SourceTemplatePair.Value->bIsActorTemplate )
+		{
+			FMigratedTemplatePairType& Result = Results.AddDefaulted_GetRef();
+
+			TStrongObjectPtr< UDatasmithObjectTemplate > SourceTemplate{ NewObject< UDatasmithObjectTemplate >(GetTransientPackage(), SourceTemplatePair.Key.Get()) }; // The SourceTemplate is the one we will persist so set its outer as DestinationObject
+
+			SourceTemplate->Load(SourceObject);
+
+			if ( ReferencesToRemap )
+			{
+				FixReferencesForObject(SourceTemplate.Get(), *ReferencesToRemap);
+			}
+
+			Result.Key = SourceTemplate;
+
+			if (DestinationObject && !DestinationObject->IsPendingKillOrUnreachable())
+			{
+				Result.Value = TStrongObjectPtr< UDatasmithObjectTemplate >(UDatasmithObjectTemplate::GetDifference( DestinationObject, SourceTemplate.Get()));
+			}
+			else
+			{
+				Result.Value = SourceTemplate;
+			}
+		}
+	}
+
+	return Results;
+}
+
+/**
+	* Applies the templates created from MigrateTemplates to DestinationObject.
+	*
+	* For an Object A that should be duplicated over an existing A', for which we want to keep the Datasmith overrides:
+	* - Call MigrateTemplates(A, A')
+	* - Duplicate A over A'
+	* - ApplyMigratedTemplates(A')
+	*/
+void FDatasmithImporterImpl::ApplyMigratedTemplates( TArray< FMigratedTemplatePairType >& MigratedTemplates, UObject* DestinationObject )
+{
+	for ( FMigratedTemplatePairType& MigratedTemplate : MigratedTemplates )
+	{
+		UDatasmithObjectTemplate* SourceTemplate = MigratedTemplate.Key.Get();
+		UDatasmithObjectTemplate* DestinationTemplate = MigratedTemplate.Value.Get();
+
+		DestinationTemplate->Apply(DestinationObject, true); // Restore the overrides
+		FDatasmithObjectTemplateUtils::SetObjectTemplate(DestinationObject, SourceTemplate); // Set SourceTemplate as our template so that any differences are considered overrides
+
+	}
+}
+
+UObject* FDatasmithImporterImpl::FinalizeAsset( UObject* SourceAsset, const TCHAR* AssetPath, UObject* ExistingAsset, TMap< UObject*, UObject* >* ReferencesToRemap )
+{
+	if ( ReferencesToRemap )
+	{
+		FixReferencesForObject( SourceAsset, *ReferencesToRemap );
+	}
+
+	TArray< FMigratedTemplatePairType > MigratedTemplates = MigrateTemplates( SourceAsset, ExistingAsset, ReferencesToRemap, false );
+
+	UObject* FinalAsset = PublicizeAsset( SourceAsset, AssetPath, ExistingAsset );
+
+	ApplyMigratedTemplates( MigratedTemplates, FinalAsset );
+
+	if ( ReferencesToRemap && SourceAsset && SourceAsset != FinalAsset )
+	{
+		ReferencesToRemap->Add( SourceAsset, FinalAsset );
+	}
+
+	return FinalAsset;
+}
+
+FDatasmithImporterImpl::FActorWriter::FActorWriter( UObject* Object, TArray< uint8 >& Bytes )
+	: FObjectWriter( Bytes )
+{
+	SetIsLoading(false);
+	SetIsSaving(true);
+	SetIsPersistent(false);
+
+	Object->Serialize(*this); // virtual call in ctr -> final class
+}
+
+bool FDatasmithImporterImpl::FActorWriter::ShouldSkipProperty(const FProperty* InProperty) const
+{
+	bool bSkip = false;
+
+	if ( InProperty->IsA< FObjectPropertyBase >() )
+	{
+		bSkip = true;
+	}
+	else if ( InProperty->HasAnyPropertyFlags(CPF_Transient) || !InProperty->HasAnyPropertyFlags(CPF_Edit | CPF_Interp) )
+	{
+		bSkip = true;
+	}
+
+	return bSkip;
+}
+
+FDatasmithImporterImpl::FComponentWriter::FComponentWriter( UObject* Object, TArray< uint8 >& Bytes )
+	: FObjectWriter( Bytes )
+{
+	SetIsLoading(false);
+	SetIsSaving(true);
+	SetIsPersistent(false);
+
+	Object->Serialize(*this); // call in ctr -> final class
+}
+
+bool FDatasmithImporterImpl::FComponentWriter::ShouldSkipProperty(const FProperty* InProperty) const
+{
+	bool bSkip = false;
+
+	if ( InProperty->HasAnyPropertyFlags(CPF_Transient) || !InProperty->HasAnyPropertyFlags(CPF_Edit | CPF_Interp) )
+	{
+		bSkip = true;
+	}
+
+	return bSkip;
+}
+
+void FDatasmithImporterImpl::DeleteImportSceneActorIfNeeded(FDatasmithActorImportContext& ActorContext, bool bForce)
+{
+	ADatasmithSceneActor*& ImportSceneActor = ActorContext.ImportSceneActor;
+	if ( !ActorContext.FinalSceneActors.Contains(ImportSceneActor) || bForce )
+	{
+		if ( ImportSceneActor )
+		{
+
+			TArray< TSoftObjectPtr< AActor > > RelatedActors;
+			ImportSceneActor->RelatedActors.GenerateValueArray( RelatedActors );
+
+			ImportSceneActor->Scene = nullptr;
+			ImportSceneActor->RelatedActors.Empty();
+
+			while(RelatedActors.Num() > 0)
+			{
+				TSoftObjectPtr< AActor > ActorPtr = RelatedActors.Pop(false);
+				if(AActor* RelatedActor = ActorPtr.Get())
+				{
+					FDatasmithImporterUtils::DeleteActor( *RelatedActor );
+				}
+			}
+
+			FDatasmithImporterUtils::DeleteActor( *ImportSceneActor );
+
+			// Null also the ImportSceneActor from the Actor Context because it's a ref to it.
+			ImportSceneActor = nullptr;
+		}
+	}
+}
+
+UActorComponent* FDatasmithImporterImpl::PublicizeComponent(UActorComponent& SourceComponent, UActorComponent* DestinationComponent, AActor& DestinationActor, TMap< UObject*, UObject* >& ReferencesToRemap, USceneComponent* DestinationParent )
+{
+	TRACE_CPUPROFILER_EVENT_SCOPE(FDatasmithImporterImpl::PublicizeComponent);
+
+	if ( !SourceComponent.HasAnyFlags( RF_Transient | RF_TextExportTransient ) )
+	{
+		if (!DestinationComponent || DestinationComponent->IsPendingKillOrUnreachable())
+		{
+			if (DestinationComponent)
+			{
+				// Change the name of the old component so that the new object won't recycle the old one.
+				DestinationComponent->Rename(nullptr, nullptr, REN_DontCreateRedirectors | REN_NonTransactional);
+			}
+
+			if ( UActorComponent* OldComponent = static_cast<UActorComponent*>( FindObjectWithOuter( &DestinationActor, UActorComponent::StaticClass(), SourceComponent.GetFName() ) ) )
+			{
+				OldComponent->DestroyComponent( true );
+				// Change the name of the old component so that the new object won't recycle the old one.
+				OldComponent->Rename(nullptr, nullptr, REN_DontCreateRedirectors | REN_NonTransactional);
+			}
+
+			// Create a new component
+			DestinationComponent = NewObject<UActorComponent>(&DestinationActor, SourceComponent.GetClass(), SourceComponent.GetFName(), RF_Transactional);
+			DestinationActor.AddInstanceComponent(DestinationComponent);
+
+			if (USceneComponent* NewSceneComponent = Cast<USceneComponent>(DestinationComponent))
+			{
+				if ( !DestinationActor.GetRootComponent() )
+				{
+					DestinationActor.SetRootComponent(NewSceneComponent);
+				}
+				if ( DestinationParent )
+				{
+					NewSceneComponent->AttachToComponent( DestinationParent, FAttachmentTransformRules::KeepRelativeTransform );
+				}
+			}
+
+			DestinationComponent->RegisterComponent();
+		}
+
+		check(DestinationComponent);
+
+		// Copy component data
+		{
+			TArray< uint8 > Bytes;
+			FComponentWriter ObjectWriter(&SourceComponent, Bytes);
+			FObjectReader ObjectReader(DestinationComponent, Bytes);
+		}
+
+		FixReferencesForObject(DestinationComponent, ReferencesToRemap);
+
+		// #ueent_todo: we shouldn't be copying instanced object pointers in the first place
+		UDatasmithAssetUserData* SourceAssetUserData = DestinationComponent->GetAssetUserData< UDatasmithAssetUserData >();
+
+		if ( SourceAssetUserData )
+		{
+			UDatasmithAssetUserData* DestinationAssetUserData = DuplicateObject< UDatasmithAssetUserData >( SourceAssetUserData, DestinationComponent );
+			DestinationComponent->RemoveUserDataOfClass( UDatasmithAssetUserData::StaticClass() );
+			DestinationComponent->AddAssetUserData( DestinationAssetUserData );
+		}
+
+		ReferencesToRemap.Add( &SourceComponent, DestinationComponent );
+
+		return DestinationComponent;
+	}
+
+	return nullptr;
+}
+
+void FDatasmithImporterImpl::FinalizeSceneComponent(FDatasmithImportContext& ImportContext, USceneComponent& SourceComponent, AActor& DestinationActor, USceneComponent* DestinationParent, TMap<UObject *, UObject *>& ReferencesToRemap )
+{
+	USceneComponent* DestinationComponent = static_cast<USceneComponent*>( FindObjectWithOuter( &DestinationActor, SourceComponent.GetClass(), SourceComponent.GetFName() ) );
+	FName SourceComponentDatasmithId = FDatasmithImporterUtils::GetDatasmithElementId(&SourceComponent);
+
+	if ( SourceComponentDatasmithId.IsNone() )
+	{
+		// This component is not tracked by datasmith
+		if ( !DestinationComponent || DestinationComponent->IsPendingKillOrUnreachable() )
+		{
+			DestinationComponent = static_cast<USceneComponent*> ( PublicizeComponent(SourceComponent, DestinationComponent, DestinationActor, ReferencesToRemap, DestinationParent) );
+			if ( DestinationComponent )
+			{
+				// Put back the components in a proper state
+				DestinationComponent->UpdateComponentToWorld();
+			}
+		}
+	}
+	else
+	{
+		check(ImportContext.ActorsContext.CurrentTargetedScene);
+
+		TArray< FMigratedTemplatePairType > MigratedTemplates = MigrateTemplates(
+			&SourceComponent,
+			DestinationComponent,
+			&ReferencesToRemap,
+			false
+		);
+
+		DestinationComponent = static_cast< USceneComponent* > ( PublicizeComponent( SourceComponent, DestinationComponent, DestinationActor, ReferencesToRemap, DestinationParent ) );
+
+		if ( DestinationComponent )
+		{
+			// Put back the components in a proper state (Unfortunately without this the set relative transform might not work)
+			DestinationComponent->UpdateComponentToWorld();
+			ApplyMigratedTemplates( MigratedTemplates, DestinationComponent );
+			DestinationComponent->PostEditChange();
+		}
+	}
+
+	USceneComponent* AttachParentForChildren = DestinationComponent  ? DestinationComponent : DestinationParent;
+	for ( USceneComponent* Child : SourceComponent.GetAttachChildren() )
+	{
+		// Only finalize components that are from the same actor
+		if ( Child && Child->GetOuter() == SourceComponent.GetOuter() )
+		{
+			FinalizeSceneComponent( ImportContext, *Child, DestinationActor, AttachParentForChildren, ReferencesToRemap );
+		}
+	}
+}
+
+void FDatasmithImporterImpl::FinalizeComponents(FDatasmithImportContext& ImportContext, AActor& SourceActor, AActor& DestinationActor, TMap<UObject *, UObject *>& ReferencesToRemap)
+{
+	USceneComponent* ParentComponent = nullptr;
+
+	// Find the parent component
+	UObject** ObjectPtr = ReferencesToRemap.Find( SourceActor.GetRootComponent()->GetAttachParent() );
+	if ( ObjectPtr )
+	{
+		ParentComponent = Cast<USceneComponent>( *ObjectPtr );
+	}
+
+	// Finalize the scene components recursively
+	{
+		USceneComponent* RootComponent = SourceActor.GetRootComponent();
+		if ( RootComponent )
+		{
+			FinalizeSceneComponent( ImportContext, *RootComponent, DestinationActor, ParentComponent, ReferencesToRemap );
+		}
+	}
+
+
+	for ( UActorComponent* SourceComponent : SourceActor.GetComponents() )
+	{
+		// Only the non scene component haven't been finalize
+		if ( SourceComponent && !SourceComponent->GetClass()->IsChildOf<USceneComponent>() )
+		{
+			UActorComponent* DestinationComponent = static_cast< UActorComponent* >( FindObjectWithOuter( &DestinationActor, SourceComponent->GetClass(), SourceComponent->GetFName() ) );
+			if (!DestinationComponent)
+			{
+				DestinationComponent = PublicizeComponent( *SourceComponent, DestinationComponent, DestinationActor, ReferencesToRemap );
+			}
+		}
+	}
+}
+
+void FDatasmithImporterImpl::GatherUnsupportedVirtualTexturesAndMaterials(const TMap<TSharedRef< IDatasmithBaseMaterialElement >, UMaterialInterface*>& ImportedMaterials, TSet<UTexture2D*>& VirtualTexturesToConvert, TArray<UMaterial*>& MaterialsToRefreshAfterVirtualTextureConversion)
+{
+	//Multimap cache to avoid parsing the same base material multiple times.
+	TMultiMap<UMaterial*, FMaterialParameterInfo> TextureParametersToConvertMap;
+
+	//Loops through all imported material instances and add to VirtualTexturesToConvert all the texture parameters that don't support texturing in the base material.
+	for (const TPair< TSharedRef< IDatasmithBaseMaterialElement >, UMaterialInterface* >& ImportedMaterialPair : ImportedMaterials)
+	{
+		UMaterialInterface* CurrentMaterialInterface = ImportedMaterialPair.Value;
+		UMaterial* BaseMaterial = CurrentMaterialInterface->GetMaterial();
+
+		if(!TextureParametersToConvertMap.Contains(BaseMaterial))
+		{
+			bool bRequiresTextureCheck = false;
+			TArray<FMaterialParameterInfo> OutParameterInfo;
+			TArray<FGuid> Guids;
+			BaseMaterial->GetAllTextureParameterInfo(OutParameterInfo, Guids);
+
+			for (int32 ParameterInfoIndex = 0; ParameterInfoIndex < OutParameterInfo.Num(); ++ParameterInfoIndex)
+			{
+				UTexture* TextureParameter = nullptr;
+
+				if (BaseMaterial->GetTextureParameterValue(OutParameterInfo[ParameterInfoIndex], TextureParameter) && VirtualTexturesToConvert.Contains(Cast<UTexture2D>(TextureParameter)))
+				{
+					bRequiresTextureCheck = true;
+					TextureParametersToConvertMap.Add(BaseMaterial, OutParameterInfo[ParameterInfoIndex]);
+				}
+			}
+
+			if (bRequiresTextureCheck)
+			{
+				MaterialsToRefreshAfterVirtualTextureConversion.Add(BaseMaterial);
+			}
+			else
+			{
+				//Adding a dummy MaterialParameterInfo so that we don't have to parse this Base Material again.
+				TextureParametersToConvertMap.Add(BaseMaterial, FMaterialParameterInfo());
+
+				//If no unsupported texture parameters were found, it's possible that a texture needing conversion is simply not exposed as a parameter, so we still need to check for those.
+				for (UObject* ReferencedTexture : BaseMaterial->GetCachedExpressionData().ReferencedTextures)
+				{
+					if (VirtualTexturesToConvert.Contains(Cast<UTexture2D>(ReferencedTexture)))
+					{
+						MaterialsToRefreshAfterVirtualTextureConversion.Add(BaseMaterial);
+						break;
+					}
+				}
+			}
+		}
+
+		for (auto ParameterInfoIt = TextureParametersToConvertMap.CreateKeyIterator(BaseMaterial); ParameterInfoIt; ++ParameterInfoIt)
+		{
+			UTexture* TextureParameter = nullptr;
+
+			if (CurrentMaterialInterface->GetTextureParameterValue(ParameterInfoIt.Value(), TextureParameter) &&
+				TextureParameter && TextureParameter->VirtualTextureStreaming)
+			{
+				if (UTexture2D* TextureToConvert = Cast<UTexture2D>(TextureParameter))
+				{
+					VirtualTexturesToConvert.Add(TextureToConvert);
+				}
+			}
+		}
+	}
+}
+
+void FDatasmithImporterImpl::ConvertUnsupportedVirtualTexture(FDatasmithImportContext& ImportContext, TSet<UTexture2D*>& VirtualTexturesToConvert, const TMap<UObject*, UObject*>& ReferencesToRemap)
+{
+	TArray<UMaterial*> MaterialsToRefreshAfterVirtualTextureConversion;
+	GatherUnsupportedVirtualTexturesAndMaterials(ImportContext.ImportedMaterials, ImportContext.AssetsContext.VirtualTexturesToConvert, MaterialsToRefreshAfterVirtualTextureConversion);
+
+	if (VirtualTexturesToConvert.Num() != 0)
+	{
+		for (UTexture2D*& TextureToConvert : VirtualTexturesToConvert)
+		{
+			if (UObject* const* RemappedTexture = ReferencesToRemap.Find(TextureToConvert))
+			{
+				TextureToConvert = Cast<UTexture2D>(*RemappedTexture);
+			}
+
+			ImportContext.LogWarning(FText::Format(LOCTEXT("DatasmithVirtualTextureConverted", "The imported texture {0} could not be imported as texture as it is not supported in all the materials using it."), FText::FromString(TextureToConvert->GetName())));
+		}
+
+		for (int32 MaterialIndex = 0; MaterialIndex < MaterialsToRefreshAfterVirtualTextureConversion.Num(); ++MaterialIndex)
+		{
+			if (UObject* const* RemappedMaterial = ReferencesToRemap.Find(MaterialsToRefreshAfterVirtualTextureConversion[MaterialIndex]))
+			{
+				MaterialsToRefreshAfterVirtualTextureConversion[MaterialIndex] = Cast<UMaterial>(*RemappedMaterial);
+			}
+		}
+
+		TArray<UTexture2D*> TexturesToConvertList(VirtualTexturesToConvert.Array());
+		IAssetTools& AssetTools = FModuleManager::LoadModuleChecked<FAssetToolsModule>("AssetTools").Get();
+		AssetTools.ConvertVirtualTextures(TexturesToConvertList, true, &MaterialsToRefreshAfterVirtualTextureConversion);
+	}
+}
+
+FDatasmithImporterImpl::FScopedFinalizeActorChanges::FScopedFinalizeActorChanges(AActor* InFinalizedActor, FDatasmithImportContext& InImportContext)
+	: ImportContext(InImportContext)
+	, FinalizedActor(InFinalizedActor)
+{
+	// In order to allow modification on components owned by ExistingActor, unregister all of them
+	FinalizedActor->UnregisterAllComponents( /* bForReregister = */true);
+
+	//Some new components might be created when finalizing the actor, only validate those that we unregistered.
+	for (UActorComponent* Component : FinalizedActor->GetComponents())
+	{
+		ComponentsToValidate.Add(Component);
+	}
+}
+
+FDatasmithImporterImpl::FScopedFinalizeActorChanges::~FScopedFinalizeActorChanges()
+{
+	for (UActorComponent* Component : FinalizedActor->GetComponents())
+	{
+		if (Component->IsRegistered() && ComponentsToValidate.Contains(Component))
+		{
+			ensureMsgf(false, TEXT("All components should still be unregistered at this point. Otherwise some datasmith templates might not have been applied properly."));
+			break;
+		}
+	}
+
+	const FQuat PreviousRotation(FinalizedActor->GetRootComponent()->GetRelativeTransform().GetRotation());
+	FinalizedActor->PostEditChange();
+	FinalizedActor->RegisterAllComponents();
+
+	const bool bHasPostEditChangeModifiedRotation = !PreviousRotation.Equals(FinalizedActor->GetRootComponent()->GetRelativeTransform().GetRotation());
+	if (bHasPostEditChangeModifiedRotation)
+	{
+		//SingularityThreshold value is comming from the FQuat::Rotator() function, but is more permissive because the rotation is already diverging before the singularity threshold is reached.
+		const float SingularityThreshold = 0.4999f;
+		const float SingularityTest = PreviousRotation.Z * PreviousRotation.X - PreviousRotation.W * PreviousRotation.Y;
+		const AActor* RootSceneActor = ImportContext.ActorsContext.ImportSceneActor;
+
+		if (FinalizedActor != RootSceneActor
+			&& FMath::Abs(SingularityTest) > SingularityThreshold)
+		{
+			//This is a warning to explain the edge-case of UE-75467 while it's being fixed.
+			FFormatNamedArguments FormatArgs;
+			FormatArgs.Add(TEXT("ActorName"), FText::FromName(FinalizedActor->GetFName()));
+			ImportContext.LogWarning(FText::GetEmpty())
+				->AddToken(FUObjectToken::Create(FinalizedActor))
+				->AddToken(FTextToken::Create(FText::Format(LOCTEXT("UnsupportedRotationValueError", "The actor '{ActorName}' has a rotation value pointing to either (0, 90, 0) or (0, -90, 0)."
+					"This is an edge case that is not well supported in Unreal and can cause incorrect results."
+					"In those cases, it is recommended to bake the actor's transform into the mesh at export."), FormatArgs)));
+		}
+	}
+}
+
+bool FDatasmithImporterImpl::CheckAssetPersistenceValidity(const FString& PackageName, FDatasmithImportContext& ImportContext)
+{
+	// Check that package can be saved
+	const FString BasePackageFileName = FPackageName::LongPackageNameToFilename( PackageName );
+	const FString AbsolutePathToAsset = FPaths::ConvertRelativePathToFull( BasePackageFileName );
+
+	// Create fake filename of same length of final asset file name to test ability to write
+	const FString FakeAbsolutePathToAsset = AbsolutePathToAsset + TEXT( ".uasset" );
+
+	bool bValidPersistence = true;
+
+	// Verify asset file name does not exceed OS' maximum path length
+	if( FPlatformMisc::GetMaxPathLength() < FakeAbsolutePathToAsset.Len() )
+	{
+		ImportContext.LogWarning(FText::Format(LOCTEXT("DatasmithImportInvalidLength", "Saving may partially fail because path for asset {0} is too long. Rename before saving."), FText::FromString( PackageName )));
+		bValidPersistence = false;
+	}
+	// Verify user can overwrite existing file
+	else if( IFileManager::Get().FileExists( *FakeAbsolutePathToAsset ) )
+	{
+		FFileStatData FileStatData = IFileManager::Get().GetStatData( *FakeAbsolutePathToAsset );
+		if ( FileStatData.bIsReadOnly )
+		{
+			// Check to see if the file is not under source control
+			bool bWarnUser = true;
+
+			ISourceControlProvider& SourceControlProvider = ISourceControlModule::Get().GetProvider();
+			if (SourceControlProvider.IsAvailable() && SourceControlProvider.IsEnabled())
+			{
+				SourceControlProvider.Execute(ISourceControlOperation::Create<FUpdateStatus>(), FakeAbsolutePathToAsset);
+				FSourceControlStatePtr SourceControlState = SourceControlProvider.GetState(FakeAbsolutePathToAsset, EStateCacheUsage::Use);
+				if( SourceControlState.IsValid() && SourceControlState->CanCheckout() )
+				{
+					// User will be prompted to check out this file when he/she saves the asset. No need to warn.
+					bWarnUser = false;
+				}
+			}
+
+			if(bWarnUser)
+			{
+				ImportContext.LogWarning(FText::Format(LOCTEXT("DatasmithImportInvalidSaving", "Saving may partially fail because file asset {0} cannot be overwritten. Check your privileges."), FText::FromString( PackageName )));
+				bValidPersistence = false;
+			}
+		}
+	}
+	// Verify user has privileges to write in folder where asset file will be stored
+	else
+	{
+		// We can't just check for the target content folders with IFileManager::Get().GetStatData here as those will
+		// only be created when UUnrealEdEngine::GetWarningStateForWritePermission is called to check for
+		// write permissions the first time, as the result is cached in GUnrealEd->PackagesCheckedForWritePermission.
+		// To check for permission, we need to first check this cache, and if the PackageName hasn't been
+		// checked yet, we need to replicate what UUnrealEdEngine::GetWarningStateForWritePermission does
+		EWriteDisallowedWarningState WarningState = EWriteDisallowedWarningState::WDWS_MAX;
+		if (GUnrealEd != nullptr && GUnrealEd->PackagesCheckedForWritePermission.Find(PackageName))
+		{
+			WarningState = (EWriteDisallowedWarningState)*GUnrealEd->PackagesCheckedForWritePermission.Find(PackageName);
+		}
+		else if (FFileHelper::SaveStringToFile(TEXT("Write Test"), *FakeAbsolutePathToAsset))
+		{
+			// We can successfully write to the folder containing the package.
+			// Delete the temp file.
+			IFileManager::Get().Delete(*FakeAbsolutePathToAsset);
+			WarningState = EWriteDisallowedWarningState::WDWS_WarningUnnecessary;
+		}
+
+		if(WarningState != EWriteDisallowedWarningState::WDWS_WarningUnnecessary)
+		{
+			ImportContext.LogWarning(FText::Format(LOCTEXT("DatasmithImportInvalidFolder", "Cannot write in folder {0} to store asset {1}. Check access to folder."), FText::FromString( FPaths::GetPath( FakeAbsolutePathToAsset ) ), FText::FromString( PackageName )));
+			bValidPersistence = false;
+		}
+	}
+
+	// Check that package can be cooked
+	// Value for MaxGameNameLen directly taken from ContentBrowserUtils::GetPackageLengthForCooking
+	constexpr int32 MaxGameNameLen = 20;
+
+	// Pad out the game name to the maximum allowed
+	const FString GameName = FApp::GetProjectName();
+	FString GameNamePadded = GameName;
+	while (GameNamePadded.Len() < MaxGameNameLen)
+	{
+		GameNamePadded += TEXT(" ");
+	}
+
+	const FString AbsoluteGamePath = FPaths::ConvertRelativePathToFull(FPaths::ProjectDir());
+	const FString AbsoluteGameCookPath = AbsoluteGamePath / TEXT("Saved") / TEXT("Cooked") / TEXT("WindowsNoEditor") / GameName;
+
+	FString AssetPathWithinCookDir = AbsolutePathToAsset;
+	FPaths::RemoveDuplicateSlashes(AssetPathWithinCookDir);
+	AssetPathWithinCookDir.RemoveFromStart(AbsoluteGamePath, ESearchCase::CaseSensitive);
+
+	// Test that the package can be cooked based on the current project path
+	FString AbsoluteCookPathToAsset = AbsoluteGameCookPath / AssetPathWithinCookDir;
+
+	AbsoluteCookPathToAsset.ReplaceInline(*GameName, *GameNamePadded, ESearchCase::CaseSensitive);
+
+	// Get the longest path allowed by the system or use 260 as the longest which is the shortest max path of any platforms that support cooking
+	const int32 MaxCookPath = GetDefault<UEditorExperimentalSettings>()->bEnableLongPathsSupport ? FPlatformMisc::GetMaxPathLength() : 260 /* MAX_PATH */;
+
+	if (AbsoluteCookPathToAsset.Len() > MaxCookPath)
+	{
+		ImportContext.LogWarning(FText::Format(LOCTEXT("DatasmithImportInvalidCooking", "Cooking may fail because path for asset {0} is too long. Rename before cooking."), FText::FromString(PackageName)));
+		bValidPersistence = false;
+	}
+
+	return bValidPersistence;
+}
+
+#undef LOCTEXT_NAMESPACE
diff --git a/Engine/Plugins/Enterprise/DatasmithImporter/Source/DatasmithImporter/Private/Utility/DatasmithImporterImpl.h b/Engine/Plugins/Enterprise/DatasmithImporter/Source/DatasmithImporter/Private/Utility/DatasmithImporterImpl.h
new file mode 100644
--- /dev/null
+++ b/Engine/Plugins/Enterprise/DatasmithImporter/Source/DatasmithImporter/Private/Utility/DatasmithImporterImpl.h
@@ -0,0 +1,142 @@
+// Copyright Epic Games, Inc. All Rights Reserved.
+
+#pragma once
+
+//#include "DatasmithImporter.h"
+
+//#include "DatasmithActorImporter.h"
+//#include "DatasmithAdditionalData.h"
+//#include "DatasmithAssetImportData.h"
+//#include "DatasmithAssetUserData.h"
+//#include "DatasmithCameraImporter.h"
+//#include "DatasmithImportContext.h"
+//#include "DatasmithLevelSequenceImporter.h"
+//#include "DatasmithLevelVariantSetsImporter.h"
+//#include "DatasmithLightImporter.h"
+//#include "DatasmithMaterialImporter.h"
+//#include "DatasmithPayload.h"
+//#include "DatasmithPostProcessImporter.h"
+//#include "DatasmithScene.h"
+//#include "DatasmithSceneActor.h"
+//#include "DatasmithSceneFactory.h"
+//#include "DatasmithStaticMeshImporter.h"
+//#include "DatasmithTranslator.h"
+//#include "DatasmithTextureImporter.h"
+//#include "IDatasmithSceneElements.h"
+//#include "DatasmithAnimationElements.h"
+//#include "LevelVariantSets.h"
+//#include "ObjectTemplates/DatasmithObjectTemplate.h"
+//#include "Utility/DatasmithImporterUtils.h"
+//#include "Utility/DatasmithTextureResize.h"
+
+#include "CoreMinimal.h"
+
+#include "Serialization/ObjectWriter.h"
+#include "UObject/StrongObjectPtr.h"
+
+class AActor;
+class FFeedbackContext;
+class IDatasmithBaseMaterialElement;
+class UActorComponent;
+class UDatasmithObjectTemplate;
+class UMaterial;
+class UMaterialInterface;
+class USceneComponent;
+class UTexture2D;
+
+struct FDatasmithActorImportContext;
+struct FDatasmithImportContext;
+struct FScopedSlowTask;
+
+class FDatasmithImporterImpl
+{
+public:
+	static void ReportProgress(FScopedSlowTask* SlowTask, const float ExpectedWorkThisFrame, const FText& Text);
+
+	static void ReportProgress(FScopedSlowTask* SlowTask, const float ExpectedWorkThisFrame, FText&& Text);
+
+	static bool HasUserCancelledTask(FFeedbackContext* FeedbackContext);
+
+	static UObject* PublicizeAsset( UObject* SourceAsset, const TCHAR* DestinationPath, UObject* ExistingAsset );
+
+	/**
+	 * Verifies the input asset can successfully be saved and/or cooked.
+	 * The code below is heavily inspired from:
+	 *			- ContentBrowserUtils::IsValidPackageForCooking
+	 * #ueent_todo Make ContentBrowserUtils::IsValidPackageForCooking public. Not game to do it so late in 4.22
+	 */
+	static bool DATASMITHIMPORTER_API CheckAssetPersistenceValidity(const FString& PackageName, FDatasmithImportContext& ImportContext);
+
+	static void CheckAssetPersistenceValidity(UObject* Asset, FDatasmithImportContext& ImportContext);
+
+	/** Set the texture mode on each texture element based on its usage in the materials */
+	static void SetTexturesMode( FDatasmithImportContext& ImportContext );
+
+	static void CompileMaterial( UObject* Material );
+
+	static void FixReferencesForObject( UObject* Object, const TMap< UObject*, UObject* >& ReferencesToRemap );
+
+	typedef  TPair< TStrongObjectPtr< UDatasmithObjectTemplate >, TStrongObjectPtr< UDatasmithObjectTemplate > > FMigratedTemplatePairType;
+
+	/**
+	 * Creates templates to apply the values from the SourceObject on the DestinationObject.
+	 *
+	 * @returns An array of template pairs. The key is the template for the object, the value is a template to force apply to the object,
+	 *			it contains the values from the key and any overrides that were present on the DestinationObject.
+	 */
+	static TArray< FMigratedTemplatePairType > MigrateTemplates( UObject* SourceObject, UObject* DestinationObject, const TMap< UObject*, UObject* >* ReferencesToRemap, bool bIsForActor );
+
+	/**
+	 * Applies the templates created from MigrateTemplates to DestinationObject.
+	 *
+	 * For an Object A that should be duplicated over an existing A', for which we want to keep the Datasmith overrides:
+	 * - Call MigrateTemplates(A, A')
+	 * - Duplicate A over A'
+	 * - ApplyMigratedTemplates(A')
+	 */
+	static void ApplyMigratedTemplates( TArray< FMigratedTemplatePairType >& MigratedTemplates, UObject* DestinationObject );
+
+	static UObject* FinalizeAsset( UObject* SourceAsset, const TCHAR* AssetPath, UObject* ExistingAsset, TMap< UObject*, UObject* >* ReferencesToRemap );
+
+	class FActorWriter final: public FObjectWriter
+	{
+	public:
+		FActorWriter( UObject* Object, TArray< uint8 >& Bytes );
+
+		virtual bool ShouldSkipProperty(const FProperty* InProperty) const override;
+	};
+
+	class FComponentWriter final : public FObjectWriter
+	{
+	public:
+		FComponentWriter( UObject* Object, TArray< uint8 >& Bytes );
+
+		virtual bool ShouldSkipProperty(const FProperty* InProperty) const override;
+	};
+
+	static void DeleteImportSceneActorIfNeeded(FDatasmithActorImportContext& ActorContext, bool bForce = false);
+
+	static UActorComponent* PublicizeComponent(UActorComponent& SourceComponent, UActorComponent* DestinationComponent, AActor& DestinationActor, TMap< UObject*, UObject* >& ReferencesToRemap, USceneComponent* DestinationParent = nullptr );
+
+	static void FinalizeSceneComponent(FDatasmithImportContext& ImportContext, USceneComponent& SourceComponent, AActor& DestinationActor, USceneComponent* DestinationParent, TMap<UObject *, UObject *>& ReferencesToRemap );
+
+	static void FinalizeComponents(FDatasmithImportContext& ImportContext, AActor& SourceActor, AActor& DestinationActor, TMap<UObject *, UObject *>& ReferencesToRemap);
+
+	static void GatherUnsupportedVirtualTexturesAndMaterials(const TMap<TSharedRef< IDatasmithBaseMaterialElement >, UMaterialInterface*>& ImportedMaterials, TSet<UTexture2D*>& VirtualTexturesToConvert, TArray<UMaterial*>& MaterialsToRefreshAfterVirtualTextureConversion);
+
+	static void ConvertUnsupportedVirtualTexture(FDatasmithImportContext& ImportContext, TSet<UTexture2D*>& VirtualTexturesToConvert, const TMap<UObject*, UObject*>& ReferencesToRemap);
+
+	/**
+	 * Helper struct used to manage and validate the actor's state when modifying it during the finalization.
+	 */
+	struct FScopedFinalizeActorChanges
+	{
+		FDatasmithImportContext& ImportContext;
+		TSet<UActorComponent*> ComponentsToValidate;
+		AActor* FinalizedActor;
+
+		FScopedFinalizeActorChanges(AActor* InFinalizedActor, FDatasmithImportContext& InImportContext);
+
+		~FScopedFinalizeActorChanges();
+	};
+};
diff --git a/Engine/Plugins/Enterprise/DatasmithImporter/Source/DatasmithImporter/Private/Utility/DatasmithImporterUtils.cpp b/Engine/Plugins/Enterprise/DatasmithImporter/Source/DatasmithImporter/Private/Utility/DatasmithImporterUtils.cpp
--- a/Engine/Plugins/Enterprise/DatasmithImporter/Source/DatasmithImporter/Private/Utility/DatasmithImporterUtils.cpp
+++ b/Engine/Plugins/Enterprise/DatasmithImporter/Source/DatasmithImporter/Private/Utility/DatasmithImporterUtils.cpp
@@ -157,7 +157,7 @@ TArray< ADatasmithSceneActor* > FDatasmithImporterUtils::FindSceneActors( UWorld
 		return TArray< ADatasmithSceneActor* >();
 	}
 
-	auto IsValidSceneActor = [DatasmithScene]( AActor* Actor )
+	auto IsValidSceneActor = [DatasmithScene, World]( AActor* Actor )
 	{
 		ADatasmithSceneActor* NullSceneActor = nullptr;
 
@@ -172,7 +172,8 @@ TArray< ADatasmithSceneActor* > FDatasmithImporterUtils::FindSceneActors( UWorld
 			{
 				if( ADatasmithSceneActor* SceneActor = Cast< ADatasmithSceneActor >( Actor ) )
 				{
-					return SceneActor->Scene == DatasmithScene ? SceneActor : NullSceneActor;
+					// A scene can be used by multiple loaded levels. Only the SceneActors of the current level are valid
+					return (SceneActor->Scene == DatasmithScene && SceneActor->GetLevel() == World->GetCurrentLevel()) ? SceneActor : NullSceneActor;
 				}
 
 				return NullSceneActor;
