Fix job [Build Engine Localization]

diff --git a/Engine/Plugins/Enterprise/DatasmithImporter/Source/DatasmithImporter/Private/DatasmithImporter.cpp b/Engine/Plugins/Enterprise/DatasmithImporter/Source/DatasmithImporter/Private/DatasmithImporter.cpp
--- a/Engine/Plugins/Enterprise/DatasmithImporter/Source/DatasmithImporter/Private/DatasmithImporter.cpp
+++ b/Engine/Plugins/Enterprise/DatasmithImporter/Source/DatasmithImporter/Private/DatasmithImporter.cpp
@@ -1715,31 +1715,4 @@ void FDatasmithImporter::FinalizeImport(FDatasmithImportContext& ImportContext,
 	GEditor->ResetTransaction(LOCTEXT("Reset Transaction Buffer", "Datasmith Import Finalization"));
 }
 
-void FDatasmithImporter::ConvertUnsupportedVirtualTexture(FDatasmithImportContext& ImportContext, TArray<UTexture2D*>& VirtualTexturesToConvert, TArray<UMaterial*>& MaterialsToRefreshAfterVirtualTextureConversion, const TMap<UObject*, UObject*>& ReferencesToRemap)
-{
-	if (VirtualTexturesToConvert.Num() != 0)
-	{
-		for (int32 TextureIndex = 0; TextureIndex < VirtualTexturesToConvert.Num(); ++TextureIndex)
-		{
-			UTexture2D*& TextureToConvert = VirtualTexturesToConvert[TextureIndex];
-			if (UObject* const* RemappedTexture = ReferencesToRemap.Find(TextureToConvert))
-			{
-				TextureToConvert = Cast<UTexture2D>(*RemappedTexture);
-			}
-
-			ImportContext.LogWarning(FText::Format(LOCTEXT("DatasmithVirtualTextureConverted", "The imported texture {0} could not be imported as virtual texture as it is not supported in all the materials using it."), FText::FromString(TextureToConvert->GetName())));
-		}
-		for (int32 MaterialIndex = 0; MaterialIndex < MaterialsToRefreshAfterVirtualTextureConversion.Num(); ++MaterialIndex)
-		{
-			if (UObject* const* RemappedMaterial = ReferencesToRemap.Find(MaterialsToRefreshAfterVirtualTextureConversion[MaterialIndex]))
-			{
-				MaterialsToRefreshAfterVirtualTextureConversion[MaterialIndex] = Cast<UMaterial>(*RemappedMaterial);
-			}
-		}
-
-		IAssetTools& AssetTools = FModuleManager::LoadModuleChecked<FAssetToolsModule>("AssetTools").Get();
-		AssetTools.ConvertVirtualTextures(VirtualTexturesToConvert, true, &MaterialsToRefreshAfterVirtualTextureConversion);
-	}
-}
-
 #undef LOCTEXT_NAMESPACE
diff --git a/Engine/Plugins/Enterprise/DatasmithImporter/Source/DatasmithImporter/Public/DatasmithImporter.h b/Engine/Plugins/Enterprise/DatasmithImporter/Source/DatasmithImporter/Public/DatasmithImporter.h
--- a/Engine/Plugins/Enterprise/DatasmithImporter/Source/DatasmithImporter/Public/DatasmithImporter.h
+++ b/Engine/Plugins/Enterprise/DatasmithImporter/Source/DatasmithImporter/Public/DatasmithImporter.h
@@ -200,8 +200,5 @@ private:
 
 	/* This class is designed not to have state. Instance are meaningless */
 	FDatasmithImporter() = delete;
-
-private:
-	static void ConvertUnsupportedVirtualTexture( FDatasmithImportContext& ImportContext, TArray<UTexture2D*>& VirtualTexturesToConvert, TArray<UMaterial*>& MaterialsToRefreshAfterVirtualTextureConversion, const TMap<UObject*, UObject*>& ReferencesToRemap);
 };
 
