Fixed crash when adding two assets with the same name to the same output folder
Added name sanitization when setting  level and folder names of consumer
Fixed logic in Dataprep editor when commit is not successful. Now data is restored to import state

#rnx

diff --git a/Engine/Plugins/Enterprise/DataprepEditor/Source/DataprepEditor/Private/DataprepEditor.cpp b/Engine/Plugins/Enterprise/DataprepEditor/Source/DataprepEditor/Private/DataprepEditor.cpp
--- a/Engine/Plugins/Enterprise/DataprepEditor/Source/DataprepEditor/Private/DataprepEditor.cpp
+++ b/Engine/Plugins/Enterprise/DataprepEditor/Source/DataprepEditor/Private/DataprepEditor.cpp
@@ -789,9 +789,15 @@ void FDataprepEditor::OnCommitWorld()
 	{
 		UE_LOG( LogDataprepEditor, Error, TEXT("Consumer failed...") );
 
+		// Restore Dataprep's import data
+		RestoreFromSnapshot();
+
 		// Restore 3D viewport
 		SceneViewportView->UpdateScene();
 
+		// Indicates that the pipeline has not been run on the data
+		bIsFirstRun = true;
+
 		return;
 	}
 
diff --git a/Engine/Plugins/Enterprise/DatasmithImporter/Source/DatasmithImporter/Private/DatasmithConsumer.cpp b/Engine/Plugins/Enterprise/DatasmithImporter/Source/DatasmithImporter/Private/DatasmithConsumer.cpp
--- a/Engine/Plugins/Enterprise/DatasmithImporter/Source/DatasmithImporter/Private/DatasmithConsumer.cpp
+++ b/Engine/Plugins/Enterprise/DatasmithImporter/Source/DatasmithImporter/Private/DatasmithConsumer.cpp
@@ -731,6 +731,13 @@ bool UDatasmithConsumer::SetLevelNameImplementation(const FString& InLevelName,
 		return false;
 	}
 
+	FString SanitizedName = ObjectTools::SanitizeObjectName(InLevelName);
+	if(SanitizedName != InLevelName)
+	{
+		OutReason = LOCTEXT( "DatasmithConsumer_InValidCharacters", "The level name contains invalid characters. Please enter a valid name." );
+		return false;
+	}
+
 	if(!CanCreateLevel(TargetContentFolder, InLevelName, !bIsAutomated && !IsRunningCommandlet()))
 	{
 		return false;
@@ -761,7 +768,7 @@ bool UDatasmithConsumer::CanCreateLevel(const FString& RequestedFolder, const FS
 
 	if(FDatasmithImporterUtils::CanCreateAsset(AssetPathName, UWorld::StaticClass()) == FDatasmithImporterUtils::EAssetCreationStatus::CS_CanCreate)
 	{
-		if(FDatasmithImporterImpl::CheckAssetPersistenceValidity(FPaths::GetPath(AssetPathName), *ImportContextPtr))
+		if(FDatasmithImporterImpl::CheckAssetPersistenceValidity(ObjectPath.GetLongPackageName(), *ImportContextPtr, FPackageName::GetMapPackageExtension()))
 		{
 			FString PackageFilename;
 			FPackageName::TryConvertLongPackageNameToFilename( ObjectPath.GetLongPackageName(), PackageFilename, FPackageName::GetMapPackageExtension() );
@@ -977,7 +984,7 @@ bool UDatasmithConsumer::CheckOutputDirectives()
 	{
 		if(FDatasmithImporterUtils::CanCreateAsset(AssetPathName, AssetClass) == FDatasmithImporterUtils::EAssetCreationStatus::CS_CanCreate)
 		{
-			return FDatasmithImporterImpl::CheckAssetPersistenceValidity(FPaths::GetPath(AssetPathName), *ImportContext);
+			return FDatasmithImporterImpl::CheckAssetPersistenceValidity(FPaths::GetPath(AssetPathName), *ImportContext, FPackageName::GetAssetPackageExtension());
 		}
 
 		return false;
@@ -990,6 +997,9 @@ bool UDatasmithConsumer::CheckOutputDirectives()
 
 	bool bCannotCreateAsset = false;
 
+	TSet<FString> AssetPaths;
+	AssetPaths.Reserve(Context.Assets.Num());
+
 	for(const TWeakObjectPtr< UObject >& AssetPtr : Context.Assets)
 	{
 		if(UObject* Asset = AssetPtr.Get())
@@ -1004,12 +1014,24 @@ bool UDatasmithConsumer::CheckOutputDirectives()
 				{
 					if(!CanCreateAsset(AssetSoftObjectPath.GetAssetPathString(), Asset->GetClass()))
 					{
-						const FTextFormat TextFormat(LOCTEXT( "DatasmithConsumer_CannotCreateAsset", "Cannot create asset {0}. Commit will be aborted" ));
+						const FTextFormat TextFormat(LOCTEXT( "DatasmithConsumer_CannotCreateAsset", "Cannot create asset {0}. Commit will be aborted." ));
 						const FText Message = FText::Format( TextFormat, FText::FromString(AssetSoftObjectPath.GetAssetPathString()) );
 						LogError(Message);
 
 						bCannotCreateAsset = true;
 					}
+					else if(AssetPaths.Contains(AssetSoftObjectPath.GetLongPackageName()))
+					{
+						const FTextFormat TextFormat(LOCTEXT( "DatasmithConsumer_DuplicateAsset", "Cannot create asset {0}. Another asset with the same name will be created in the same folder {1}. Commit will be aborted." ));
+						const FText Message = FText::Format( TextFormat, FText::FromString(AssetSoftObjectPath.GetAssetPathString()), FText::FromString(FPaths::Combine(TargetContentFolder, OutputFolder)) );
+						LogError(Message);
+
+						bCannotCreateAsset = true;
+					}
+					else
+					{
+						AssetPaths.Add(AssetSoftObjectPath.GetLongPackageName());
+					}
 				}
 			}
 		}
diff --git a/Engine/Plugins/Enterprise/DatasmithImporter/Source/DatasmithImporter/Private/UI/DatasmithConsumerDetails.cpp b/Engine/Plugins/Enterprise/DatasmithImporter/Source/DatasmithImporter/Private/UI/DatasmithConsumerDetails.cpp
--- a/Engine/Plugins/Enterprise/DatasmithImporter/Source/DatasmithImporter/Private/UI/DatasmithConsumerDetails.cpp
+++ b/Engine/Plugins/Enterprise/DatasmithImporter/Source/DatasmithImporter/Private/UI/DatasmithConsumerDetails.cpp
@@ -14,6 +14,7 @@
 #include "Engine/LevelStreamingAlwaysLoaded.h"
 #include "Engine/World.h"
 #include "Factories/LevelFactory.h"
+#include "Misc/FileHelper.h"
 #include "Misc/MessageDialog.h"
 #include "Misc/PackageName.h"
 #include "PropertyCustomizationHelpers.h"
@@ -27,6 +28,25 @@
 
 namespace DatasmithConsumerDetailsUtil
 {
+	// Inspired from ContentBrowserUtils::IsValidObjectPathForCreate
+	bool VerifyObjectName(const FString& ObjectName, FText& OutErrorMessage)
+	{
+		if ( !FFileHelper::IsFilenameValidForSaving(ObjectName, OutErrorMessage) )
+		{
+			// Return false to indicate that the user should enter a new name
+			return false;
+		}
+
+		// Make sure the new name only contains valid characters
+		if ( !FName::IsValidXName( ObjectName, INVALID_OBJECTNAME_CHARACTERS INVALID_LONGPACKAGE_CHARACTERS, &OutErrorMessage ) )
+		{
+			// Return false to indicate that the user should enter a new name
+			return false;
+		}
+
+		return true;
+	}
+
 	/** Helper class to force a widget to fill in a space. Copied from SDetailSingleItemRow.cpp */
 	class SConstrainedBox : public SCompoundWidget
 	{
@@ -83,6 +103,7 @@ namespace DatasmithConsumerDetailsUtil
 						.HintText(LOCTEXT("DataprepSlateHelper_ContentFolderHintText", "Set the content folder to save in"))
 						.IsReadOnly(false)
 						.OnTextCommitted(FOnTextCommitted::CreateSP(this, &SFolderProperty::OnTextCommitted))
+						.OnVerifyTextChanged(FOnVerifyTextChanged::CreateSP(this, &SFolderProperty::OnVerifyText))
 					]
 				]
 				+ SHorizontalBox::Slot()
@@ -131,6 +152,89 @@ namespace DatasmithConsumerDetailsUtil
 			}
 		}
 
+		// Inspired from ContentBrowserUtils::IsValidFolderPathForCreate
+		bool OnVerifyText(const FText& InText, FText& OutErrorMessage)
+		{
+			FString FolderPath = InText.ToString();
+
+			if( FolderPath.StartsWith( TEXT("/Content") ) )
+			{
+				FolderPath = FolderPath.Replace( TEXT( "/Content" ), TEXT( "/Game" ) );
+			}
+
+			// Check length of the folder name
+			if ( FolderPath.Len() == 0 )
+			{
+				OutErrorMessage = LOCTEXT( "InvalidFolderName_IsTooShort", "Please provide a name for this folder." );
+				return false;
+			}
+
+			if ( FolderPath.Len() > FPlatformMisc::GetMaxPathLength() )
+			{
+				OutErrorMessage = FText::Format( LOCTEXT("InvalidFolderName_TooLongForCooking", "Filename is too long ({0} characters); this may interfere with cooking for consoles. Unreal filenames should be no longer than {1} characters. Filename value: {2}" ),
+					FText::AsNumber(FolderPath.Len()), FText::AsNumber(FPlatformMisc::GetMaxPathLength()), FText::FromString(FolderPath) );
+				return false;
+			}
+
+			if(FolderPath[FolderPath.Len() - 1] == L'/')
+			{
+				return true;
+			}
+
+			FString FolderName = FPaths::GetBaseFilename(FolderPath);
+
+			if(!VerifyObjectName(FolderName, OutErrorMessage))
+			{
+				return false;
+			}
+
+
+			const FString InvalidChars = INVALID_LONGPACKAGE_CHARACTERS TEXT("/[]"); // Slash and Square brackets are invalid characters for a folder name
+
+																					 // See if the name contains invalid characters.
+			FString Char;
+			for( int32 CharIdx = 0; CharIdx < FolderName.Len(); ++CharIdx )
+			{
+				Char = FolderName.Mid(CharIdx, 1);
+
+				if ( InvalidChars.Contains(*Char) )
+				{
+					FString ReadableInvalidChars = InvalidChars;
+					ReadableInvalidChars.ReplaceInline(TEXT("\r"), TEXT(""));
+					ReadableInvalidChars.ReplaceInline(TEXT("\n"), TEXT(""));
+					ReadableInvalidChars.ReplaceInline(TEXT("\t"), TEXT(""));
+
+					OutErrorMessage = FText::Format(LOCTEXT("InvalidFolderName_InvalidCharacters", "A folder name may not contain any of the following characters: {0}"), FText::FromString(ReadableInvalidChars));
+					return false;
+				}
+			}
+
+			if(!FFileHelper::IsFilenameValidForSaving(FolderPath, OutErrorMessage))
+			{
+				return false;
+			}
+
+			FString PathOnDisk;
+			if (!FPackageName::TryConvertLongPackageNameToFilename(FolderPath, PathOnDisk))
+			{
+				OutErrorMessage = FText::Format(LOCTEXT("RenameFolderFailedDiskPath", "Folder path could not be converted to disk path: '{0}'"), FText::FromString(FolderPath));
+				return false;
+			}
+
+			// Make sure we are not creating a folder path that is too long
+			if (PathOnDisk.Len() > FPlatformMisc::GetMaxPathLength() - 32/*MAX_CLASS_NAME_LENGTH*/)
+			{
+				// The full path for the folder is too long
+				OutErrorMessage = FText::Format(LOCTEXT("RenameFolderPathTooLong",
+					"The full path for the folder is too deep, the maximum is '{0}'. Please choose a shorter name for the folder or create it in a shallower folder structure."),
+					FText::AsNumber(FPlatformMisc::GetMaxPathLength() - 32/*MAX_CLASS_NAME_LENGTH*/));
+				// Return false to indicate that the user should enter a new name for the folder
+				return false;
+			}
+
+			return true;
+		}
+
 		void UpdateContentFolderText()
 		{
 			if(UDataprepContentConsumer* Consumer = ConsumerPtr.Get())
@@ -213,16 +317,12 @@ namespace DatasmithConsumerDetailsUtil
 				.HAlign(EHorizontalAlignment::HAlign_Fill)
 				.Padding(5.0f, 2.5f, 2.0f, 2.5f)
 				[
-					// Trick to force the splitter widget to fill up the space of its parent
-					// Strongly inspired from SDetailSingleItemRow
-					//SNew(SConstrainedBox)
-					//[
-						SAssignNew(LevelTextBox, SEditableTextBox)
-						.Font(IDetailLayoutBuilder::GetDetailFont())
-						.HintText(LOCTEXT("DataprepLevelProperty_HintText", "Set the name of the level to save in"))
-						.IsReadOnly(false)
-						.OnTextCommitted(FOnTextCommitted::CreateSP(this, &SLevelProperty::OnTextCommitted))
-					//]
+					SAssignNew(LevelTextBox, SEditableTextBox)
+					.Font(IDetailLayoutBuilder::GetDetailFont())
+					.HintText(LOCTEXT("DataprepLevelProperty_HintText", "Set the name of the level to save in"))
+					.IsReadOnly(false)
+					.OnTextCommitted(FOnTextCommitted::CreateSP(this, &SLevelProperty::OnTextCommitted))
+					.OnVerifyTextChanged(FOnVerifyTextChanged::CreateSP(this, &SLevelProperty::OnVerifyText))
 				]
 			];
 
@@ -242,6 +342,11 @@ namespace DatasmithConsumerDetailsUtil
 			}
 		}
 
+		bool OnVerifyText(const FText& InText, FText& OutErrorMessage)
+		{
+			return VerifyObjectName(InText.ToString(), OutErrorMessage);
+		}
+
 		void OnTextCommitted( const FText& NewText, ETextCommit::Type CommitType)
 		{
 			if( UDatasmithConsumer* DataprepConsumer = ConsumerPtr.Get() )
diff --git a/Engine/Plugins/Enterprise/DatasmithImporter/Source/DatasmithImporter/Private/Utility/DatasmithImporterImpl.cpp b/Engine/Plugins/Enterprise/DatasmithImporter/Source/DatasmithImporter/Private/Utility/DatasmithImporterImpl.cpp
--- a/Engine/Plugins/Enterprise/DatasmithImporter/Source/DatasmithImporter/Private/Utility/DatasmithImporterImpl.cpp
+++ b/Engine/Plugins/Enterprise/DatasmithImporter/Source/DatasmithImporter/Private/Utility/DatasmithImporterImpl.cpp
@@ -56,6 +56,7 @@
 #include "Materials/MaterialInstanceConstant.h"
 #include "Misc/FeedbackContext.h"
 #include "Misc/FileHelper.h"
+#include "Misc/PackageName.h"
 #include "Misc/ScopedSlowTask.h"
 #include "Misc/UObjectToken.h"
 #include "ObjectTools.h"
@@ -184,7 +185,7 @@ void FDatasmithImporterImpl::CheckAssetPersistenceValidity(UObject* Asset, FData
 	UPackage* Package = Asset->GetOutermost();
 	const FString PackageName = Package->GetName();
 
-	CheckAssetPersistenceValidity(PackageName, ImportContext);
+	CheckAssetPersistenceValidity(PackageName, ImportContext, Asset->IsA<UWorld>() ? FPackageName::GetMapPackageExtension() : FPackageName::GetAssetPackageExtension());
 }
 
 /** Set the texture mode on each texture element based on its usage in the materials */
@@ -827,14 +828,14 @@ FDatasmithImporterImpl::FScopedFinalizeActorChanges::~FScopedFinalizeActorChange
 	}
 }
 
-bool FDatasmithImporterImpl::CheckAssetPersistenceValidity(const FString& PackageName, FDatasmithImportContext& ImportContext)
+bool FDatasmithImporterImpl::CheckAssetPersistenceValidity(const FString& PackageName, FDatasmithImportContext& ImportContext, const FString& Extension)
 {
 	// Check that package can be saved
 	const FString BasePackageFileName = FPackageName::LongPackageNameToFilename( PackageName );
 	const FString AbsolutePathToAsset = FPaths::ConvertRelativePathToFull( BasePackageFileName );
 
 	// Create fake filename of same length of final asset file name to test ability to write
-	const FString FakeAbsolutePathToAsset = AbsolutePathToAsset + TEXT( ".uasset" );
+	const FString FakeAbsolutePathToAsset = AbsolutePathToAsset + Extension;
 
 	bool bValidPersistence = true;
 
diff --git a/Engine/Plugins/Enterprise/DatasmithImporter/Source/DatasmithImporter/Private/Utility/DatasmithImporterImpl.h b/Engine/Plugins/Enterprise/DatasmithImporter/Source/DatasmithImporter/Private/Utility/DatasmithImporterImpl.h
--- a/Engine/Plugins/Enterprise/DatasmithImporter/Source/DatasmithImporter/Private/Utility/DatasmithImporterImpl.h
+++ b/Engine/Plugins/Enterprise/DatasmithImporter/Source/DatasmithImporter/Private/Utility/DatasmithImporterImpl.h
@@ -65,7 +65,7 @@ public:
 	 *			- ContentBrowserUtils::IsValidPackageForCooking
 	 * #ueent_todo Make ContentBrowserUtils::IsValidPackageForCooking public. Not game to do it so late in 4.22
 	 */
-	static bool DATASMITHIMPORTER_API CheckAssetPersistenceValidity(const FString& PackageName, FDatasmithImportContext& ImportContext);
+	static bool DATASMITHIMPORTER_API CheckAssetPersistenceValidity(const FString& PackageName, FDatasmithImportContext& ImportContext, const FString& Extension);
 
 	static void CheckAssetPersistenceValidity(UObject* Asset, FDatasmithImportContext& ImportContext);
 
