Fixed Warning msg to overwrite level appears twice

#rnx

diff --git a/Engine/Plugins/Enterprise/DatasmithImporter/Source/DatasmithImporter/Private/UI/DatasmithConsumerDetails.cpp b/Engine/Plugins/Enterprise/DatasmithImporter/Source/DatasmithImporter/Private/UI/DatasmithConsumerDetails.cpp
--- a/Engine/Plugins/Enterprise/DatasmithImporter/Source/DatasmithImporter/Private/UI/DatasmithConsumerDetails.cpp
+++ b/Engine/Plugins/Enterprise/DatasmithImporter/Source/DatasmithImporter/Private/UI/DatasmithConsumerDetails.cpp
@@ -82,6 +82,8 @@ namespace DatasmithConsumerDetailsUtil
 		void Construct(const FArguments& InArgs, UDatasmithConsumer* InConsumer)
 		{
 			ConsumerPtr = InConsumer;
+			bProcessing = false;
+
 			FSlateFontInfo FontInfo = IDetailLayoutBuilder::GetDetailFont();
 
 			TSharedRef<SWidget> BrowseButton = PropertyCustomizationHelpers::MakeBrowseButton( FSimpleDelegate::CreateSP( this, &SFolderProperty::OnBrowseContentFolder ) );
@@ -260,8 +262,15 @@ namespace DatasmithConsumerDetailsUtil
 
 		void OnTextCommitted( const FText& NewText, ETextCommit::Type CommitType)
 		{
+			if(bProcessing)
+			{
+				return;
+			}
+
 			if( UDataprepContentConsumer* DataprepConsumer = ConsumerPtr.Get() )
 			{
+				bProcessing = true;
+
 				FString NewContentFolder( NewText.ToString() );
 
 				// Replace /Content/ with /Game/ since /Content is only used for display 
@@ -289,12 +298,16 @@ namespace DatasmithConsumerDetailsUtil
 						UpdateContentFolderText();
 					}
 				}
+
+				bProcessing = false;
 			}
 		}
 
 	private:
 		TWeakObjectPtr< UDatasmithConsumer > ConsumerPtr;
 		TSharedPtr< SEditableTextBox > ContentFolderTextBox;
+		// Boolean used to avoid re-entering UI event processing
+		bool bProcessing;
 	};
 
 	class SLevelProperty : public SCompoundWidget
@@ -307,6 +320,8 @@ namespace DatasmithConsumerDetailsUtil
 		void Construct(const FArguments& InArgs, UDatasmithConsumer* InConsumer)
 		{
 			ConsumerPtr = InConsumer;
+			bProcessing = false;
+
 			FSlateFontInfo FontInfo = IDetailLayoutBuilder::GetDetailFont();
 
 			ChildSlot
@@ -349,8 +364,14 @@ namespace DatasmithConsumerDetailsUtil
 
 		void OnTextCommitted( const FText& NewText, ETextCommit::Type CommitType)
 		{
+			if(bProcessing)
+			{
+				return;
+			}
+
 			if( UDatasmithConsumer* DataprepConsumer = ConsumerPtr.Get() )
 			{
+				bProcessing = true;
 				FString NewLevelName( NewText.ToString() );
 
 				if(NewLevelName != DataprepConsumer->GetLevelName())
@@ -367,12 +388,15 @@ namespace DatasmithConsumerDetailsUtil
 						UE_LOG( LogDatasmithImport, Error, TEXT("Cannot create a level named %s - %s"), *NewLevelName, *OutReason.ToString() );
 					}
 				}
+				bProcessing = false;
 			}
 		}
 
 	private:
 		TWeakObjectPtr< UDatasmithConsumer > ConsumerPtr;
 		TSharedPtr< SEditableTextBox > LevelTextBox;
+		// Boolean used to avoid re-entering UI event processing
+		bool bProcessing;
 	};
 }
 
