Integrating CL# 11896663 from Dev-Enterprise
Fixing reimporting datasmith scene with different options making no change.

diff --git a/Engine/Plugins/Enterprise/DatasmithContent/Source/DatasmithContent/Public/DatasmithImportOptions.h b/Engine/Plugins/Enterprise/DatasmithContent/Source/DatasmithContent/Public/DatasmithImportOptions.h
--- a/Engine/Plugins/Enterprise/DatasmithContent/Source/DatasmithContent/Public/DatasmithImportOptions.h
+++ b/Engine/Plugins/Enterprise/DatasmithContent/Source/DatasmithContent/Public/DatasmithImportOptions.h
@@ -146,7 +146,14 @@ public:
 
 	bool operator == (const FDatasmithStaticMeshImportOptions& Other) const
 	{
-		return (MinLightmapResolution == Other.MinLightmapResolution && MaxLightmapResolution == Other.MaxLightmapResolution && bRemoveDegenerates == Other.bRemoveDegenerates);
+		return MinLightmapResolution == Other.MinLightmapResolution 
+			&& MaxLightmapResolution == Other.MaxLightmapResolution
+			&& bGenerateLightmapUVs == Other.bGenerateLightmapUVs
+			&& bRemoveDegenerates == Other.bRemoveDegenerates;
+	}
+
+	bool operator != (const FDatasmithStaticMeshImportOptions& Other) const {
+		return !operator==(Other);
 	}
 };
 
diff --git a/Engine/Plugins/Enterprise/DatasmithImporter/Source/DatasmithImporter/Private/DatasmithImporter.cpp b/Engine/Plugins/Enterprise/DatasmithImporter/Source/DatasmithImporter/Private/DatasmithImporter.cpp
--- a/Engine/Plugins/Enterprise/DatasmithImporter/Source/DatasmithImporter/Private/DatasmithImporter.cpp
+++ b/Engine/Plugins/Enterprise/DatasmithImporter/Source/DatasmithImporter/Private/DatasmithImporter.cpp
@@ -1363,14 +1363,16 @@ void FDatasmithImporter::FilterElementsToImport( FDatasmithImportContext& Import
 
 	// Meshes part
 	ImportContext.FilteredScene->EmptyMeshes();
+	const bool bSameStaticMeshOptions = ImportContext.Options->BaseOptions.StaticMeshOptions == ImportContext.SceneAsset->AssetImportData->BaseOptions.StaticMeshOptions;
 	const TMap< FName, TSoftObjectPtr< UStaticMesh > >& StaticMeshes = ImportContext.SceneAsset->StaticMeshes;
 	for (int32 MeshIndex = 0; MeshIndex < ImportContext.Scene->GetMeshesCount(); ++MeshIndex)
 	{
 		TSharedRef< IDatasmithMeshElement > MeshElement = ImportContext.Scene->GetMesh( MeshIndex ).ToSharedRef();
-
 		bool bNeedsReimport = true;
 		FString AssetName = MeshElement->GetName();
-		if ( StaticMeshes.Contains( MeshElement->GetName() ) )
+		
+		// If we are reimporting with different options we should not try to skip the reimport.
+		if ( bSameStaticMeshOptions && StaticMeshes.Contains( MeshElement->GetName() ) )
 		{
 			AssetName = StaticMeshes[ MeshElement->GetName() ].ToString();
 			bNeedsReimport = ElementNeedsReimport(AssetName, MeshElement, ImportContext.Options->FilePath );
