Integrating CL# 11908933 from Dev-Enterprise
Fixing crash when importing with Dataprep

diff --git a/Engine/Plugins/Enterprise/DatasmithImporter/Source/DatasmithImporter/Private/DatasmithImporter.cpp b/Engine/Plugins/Enterprise/DatasmithImporter/Source/DatasmithImporter/Private/DatasmithImporter.cpp
--- a/Engine/Plugins/Enterprise/DatasmithImporter/Source/DatasmithImporter/Private/DatasmithImporter.cpp
+++ b/Engine/Plugins/Enterprise/DatasmithImporter/Source/DatasmithImporter/Private/DatasmithImporter.cpp
@@ -1319,7 +1319,7 @@ void FDatasmithImporter::FilterElementsToImport( FDatasmithImportContext& Import
 	IAssetRegistry& AssetRegistry = FModuleManager::LoadModuleChecked<FAssetRegistryModule>(TEXT("AssetRegistry")).Get();
 
 	// No Scene asset yet, all assets of the scene must be imported
-	if (!ImportContext.SceneAsset)
+	if (!ImportContext.SceneAsset || !ImportContext.SceneAsset->AssetImportData)
 	{
 		return;
 	}
@@ -1363,7 +1363,7 @@ void FDatasmithImporter::FilterElementsToImport( FDatasmithImportContext& Import
 
 	// Meshes part
 	ImportContext.FilteredScene->EmptyMeshes();
-	const bool bSameStaticMeshOptions = ImportContext.Options->BaseOptions.StaticMeshOptions == ImportContext.SceneAsset->AssetImportData->BaseOptions.StaticMeshOptions;
+	const bool bDifferentStaticMeshOptions = ImportContext.Options->BaseOptions.StaticMeshOptions != ImportContext.SceneAsset->AssetImportData->BaseOptions.StaticMeshOptions;
 	const TMap< FName, TSoftObjectPtr< UStaticMesh > >& StaticMeshes = ImportContext.SceneAsset->StaticMeshes;
 	for (int32 MeshIndex = 0; MeshIndex < ImportContext.Scene->GetMeshesCount(); ++MeshIndex)
 	{
@@ -1371,11 +1371,11 @@ void FDatasmithImporter::FilterElementsToImport( FDatasmithImportContext& Import
 		bool bNeedsReimport = true;
 		FString AssetName = MeshElement->GetName();
 		
-		// If we are reimporting with different options we should not try to skip the reimport.
-		if ( bSameStaticMeshOptions && StaticMeshes.Contains( MeshElement->GetName() ) )
+		if ( StaticMeshes.Contains( MeshElement->GetName() ) )
 		{
 			AssetName = StaticMeshes[ MeshElement->GetName() ].ToString();
-			bNeedsReimport = ElementNeedsReimport(AssetName, MeshElement, ImportContext.Options->FilePath );
+			// If we are reimporting with different options we should not try to skip the reimport.
+			bNeedsReimport = bDifferentStaticMeshOptions || ElementNeedsReimport(AssetName, MeshElement, ImportContext.Options->FilePath );
 		}
 
 		if ( bNeedsReimport )
