Fixed crash with relative path as level name

#rnx

diff --git a/Engine/Plugins/Enterprise/DatasmithImporter/Source/DatasmithImporter/Private/DatasmithConsumer.cpp b/Engine/Plugins/Enterprise/DatasmithImporter/Source/DatasmithImporter/Private/DatasmithConsumer.cpp
--- a/Engine/Plugins/Enterprise/DatasmithImporter/Source/DatasmithImporter/Private/DatasmithConsumer.cpp
+++ b/Engine/Plugins/Enterprise/DatasmithImporter/Source/DatasmithImporter/Private/DatasmithConsumer.cpp
@@ -731,6 +731,13 @@ bool UDatasmithConsumer::SetLevelNameImplementation(const FString& InLevelName,
 		return false;
 	}
 
+	int32 Index;
+	if(InLevelName.FindChar( L'.', Index) || InLevelName.FindChar( L'/', Index))
+	{
+		OutReason = LOCTEXT( "DatasmithConsumer_IsAPath", "Path or relative path is not accepted as a folder name. Please enter a valid name." );
+		return false;
+	}
+
 	FString SanitizedName = ObjectTools::SanitizeObjectName(InLevelName);
 	if(SanitizedName != InLevelName)
 	{
@@ -854,6 +861,13 @@ bool UDatasmithConsumer::SetTargetContentFolderImplementation(const FString& InT
 		return true;
 	}
 
+	if(InTargetContentFolder.StartsWith( TEXT("..")))
+	{
+		const FText Message = LOCTEXT( "DatasmithConsumer_RelativePath", "Relative path is not accepted as level name. Please enter a valid name." );
+		LogInfo(Message);
+		return false;
+	}
+
 	if(!CanCreateLevel(InTargetContentFolder, LevelName, !bIsAutomated && !IsRunningCommandlet()))
 	{
 		return false;
diff --git a/Engine/Plugins/Enterprise/DatasmithImporter/Source/DatasmithImporter/Private/UI/DatasmithConsumerDetails.cpp b/Engine/Plugins/Enterprise/DatasmithImporter/Source/DatasmithImporter/Private/UI/DatasmithConsumerDetails.cpp
--- a/Engine/Plugins/Enterprise/DatasmithImporter/Source/DatasmithImporter/Private/UI/DatasmithConsumerDetails.cpp
+++ b/Engine/Plugins/Enterprise/DatasmithImporter/Source/DatasmithImporter/Private/UI/DatasmithConsumerDetails.cpp
@@ -90,15 +90,14 @@ namespace DatasmithConsumerDetailsUtil
 
 			ChildSlot
 			[
-				SNew(SHorizontalBox)
-				+ SHorizontalBox::Slot()
-				.FillWidth(1.0f)
-				.HAlign(EHorizontalAlignment::HAlign_Fill)
+				SNew(SVerticalBox)
+				+ SVerticalBox::Slot()
 				.Padding(5.0f, 2.5f, 2.0f, 2.5f)
 				[
-					// Trick to force the splitter widget to fill up the space of its parent
-					// Strongly inspired from SDetailSingleItemRow
-					SNew(SConstrainedBox)
+					SNew(SHorizontalBox)
+					+ SHorizontalBox::Slot()
+					.FillWidth(9.0f)
+					.HAlign(EHorizontalAlignment::HAlign_Fill)
 					[
 						SAssignNew(ContentFolderTextBox, SEditableTextBox)
 						.Font(IDetailLayoutBuilder::GetDetailFont())
@@ -107,12 +106,14 @@ namespace DatasmithConsumerDetailsUtil
 						.OnTextCommitted(FOnTextCommitted::CreateSP(this, &SFolderProperty::OnTextCommitted))
 						.OnVerifyTextChanged(FOnVerifyTextChanged::CreateSP(this, &SFolderProperty::OnVerifyText))
 					]
+					+ SHorizontalBox::Slot()
+					.FillWidth(1.0f)
+					.HAlign(EHorizontalAlignment::HAlign_Right)
+					[
+						BrowseButton
+					]
 				]
-				+ SHorizontalBox::Slot()
-				.HAlign(EHorizontalAlignment::HAlign_Right)
-				[
-					BrowseButton
-				]
+
 			];
 
 			UpdateContentFolderText();
@@ -159,11 +160,6 @@ namespace DatasmithConsumerDetailsUtil
 		{
 			FString FolderPath = InText.ToString();
 
-			if( FolderPath.StartsWith( TEXT("/Content") ) )
-			{
-				FolderPath = FolderPath.Replace( TEXT( "/Content" ), TEXT( "/Game" ) );
-			}
-
 			// Check length of the folder name
 			if ( FolderPath.Len() == 0 )
 			{
@@ -171,6 +167,17 @@ namespace DatasmithConsumerDetailsUtil
 				return false;
 			}
 
+			if(FolderPath.StartsWith( TEXT("..")))
+			{
+				OutErrorMessage = LOCTEXT( "InvalidFolderName_RelativePath", "Relative path is not accepted." );
+				return false;
+			}
+
+			if( FolderPath.StartsWith( TEXT("/Content") ) )
+			{
+				FolderPath = FolderPath.Replace( TEXT( "/Content" ), TEXT( "/Game" ) );
+			}
+
 			if ( FolderPath.Len() > FPlatformMisc::GetMaxPathLength() )
 			{
 				OutErrorMessage = FText::Format( LOCTEXT("InvalidFolderName_TooLongForCooking", "Filename is too long ({0} characters); this may interfere with cooking for consoles. Unreal filenames should be no longer than {1} characters. Filename value: {2}" ),
@@ -180,7 +187,7 @@ namespace DatasmithConsumerDetailsUtil
 
 			if(FolderPath[FolderPath.Len() - 1] == L'/')
 			{
-				return true;
+				FolderPath.LeftInline(FolderPath.Len() - 1, false);
 			}
 
 			FString FolderName = FPaths::GetBaseFilename(FolderPath);
@@ -321,6 +328,7 @@ namespace DatasmithConsumerDetailsUtil
 		{
 			ConsumerPtr = InConsumer;
 			bProcessing = false;
+			bValidText = true;
 
 			FSlateFontInfo FontInfo = IDetailLayoutBuilder::GetDetailFont();
 
@@ -359,7 +367,18 @@ namespace DatasmithConsumerDetailsUtil
 
 		bool OnVerifyText(const FText& InText, FText& OutErrorMessage)
 		{
-			return VerifyObjectName(InText.ToString(), OutErrorMessage);
+			const FString NewLevelName = InText.ToString();
+
+			int32 Index;
+			if(NewLevelName.FindChar(L'/', Index))
+			{
+				OutErrorMessage = LOCTEXT( "InvalidLevelName_RelativePath", "Path or relative path is not accepted." );
+				bValidText = false;
+			}
+
+			bValidText = VerifyObjectName(NewLevelName, OutErrorMessage);
+
+			return bValidText;
 		}
 
 		void OnTextCommitted( const FText& NewText, ETextCommit::Type CommitType)
@@ -369,6 +388,13 @@ namespace DatasmithConsumerDetailsUtil
 				return;
 			}
 
+			if(!bValidText)
+			{
+				UpdateLevelText();
+				bValidText = true;
+				return;
+			}
+
 			if( UDatasmithConsumer* DataprepConsumer = ConsumerPtr.Get() )
 			{
 				bProcessing = true;
@@ -397,6 +423,7 @@ namespace DatasmithConsumerDetailsUtil
 		TSharedPtr< SEditableTextBox > LevelTextBox;
 		// Boolean used to avoid re-entering UI event processing
 		bool bProcessing;
+		bool bValidText;
 	};
 }
 
diff --git a/Engine/Plugins/Enterprise/DatasmithImporter/Source/DatasmithImporter/Private/Utility/DatasmithImporterImpl.cpp b/Engine/Plugins/Enterprise/DatasmithImporter/Source/DatasmithImporter/Private/Utility/DatasmithImporterImpl.cpp
--- a/Engine/Plugins/Enterprise/DatasmithImporter/Source/DatasmithImporter/Private/Utility/DatasmithImporterImpl.cpp
+++ b/Engine/Plugins/Enterprise/DatasmithImporter/Source/DatasmithImporter/Private/Utility/DatasmithImporterImpl.cpp
@@ -66,6 +66,7 @@
 #include "Serialization/ObjectReader.h"
 #include "Settings/EditorExperimentalSettings.h"
 #include "SourceControlOperations.h"
+#include "Subsystems/AssetEditorSubsystem.h"
 #include "Templates/UniquePtr.h"
 #include "UObject/Package.h"
 #include "UnrealEdGlobals.h"
@@ -122,6 +123,13 @@ UObject* FDatasmithImporterImpl::PublicizeAsset( UObject* SourceAsset, const TCH
 		DestinationPackage = ExistingAsset->GetOutermost();
 	}
 
+	// Close editors opened on existing asset if applicable
+	UAssetEditorSubsystem* AssetEditorSubsystem = GEditor->GetEditorSubsystem<UAssetEditorSubsystem>();
+	if (ExistingAsset && AssetEditorSubsystem->FindEditorForAsset(ExistingAsset, false) != nullptr)
+	{
+		AssetEditorSubsystem->CloseAllEditorsForAsset(ExistingAsset);
+	}
+
 	DestinationPackage->FullyLoad();
 
 	UObject* DestinationAsset = ExistingAsset;
