Improved handling of invalid characters in level and folder widgets for consumer object

#rnx

diff --git a/Engine/Plugins/Enterprise/DatasmithImporter/Source/DatasmithImporter/Private/UI/DatasmithConsumerDetails.cpp b/Engine/Plugins/Enterprise/DatasmithImporter/Source/DatasmithImporter/Private/UI/DatasmithConsumerDetails.cpp
--- a/Engine/Plugins/Enterprise/DatasmithImporter/Source/DatasmithImporter/Private/UI/DatasmithConsumerDetails.cpp
+++ b/Engine/Plugins/Enterprise/DatasmithImporter/Source/DatasmithImporter/Private/UI/DatasmithConsumerDetails.cpp
@@ -104,7 +104,7 @@ namespace DatasmithConsumerDetailsUtil
 						.HintText(LOCTEXT("DataprepSlateHelper_ContentFolderHintText", "Set the content folder to save in"))
 						.IsReadOnly(false)
 						.OnTextCommitted(FOnTextCommitted::CreateSP(this, &SFolderProperty::OnTextCommitted))
-						.OnVerifyTextChanged(FOnVerifyTextChanged::CreateSP(this, &SFolderProperty::OnVerifyText))
+						.OnTextChanged(FOnTextChanged::CreateSP(this, &SFolderProperty::OnTextChanged))
 					]
 					+ SHorizontalBox::Slot()
 					.FillWidth(1.0f)
@@ -156,96 +156,102 @@ namespace DatasmithConsumerDetailsUtil
 		}
 
 		// Inspired from ContentBrowserUtils::IsValidFolderPathForCreate
-		bool OnVerifyText(const FText& InText, FText& OutErrorMessage)
+		void OnTextChanged(const FText& InText)
 		{
+			FText OutErrorMessage = FText::GetEmpty();
 			FString FolderPath = InText.ToString();
 
 			// Check length of the folder name
-			if ( FolderPath.Len() == 0 )
+			if ( FolderPath.Len() > 0 )
 			{
-				OutErrorMessage = LOCTEXT( "InvalidFolderName_IsTooShort", "Please provide a name for this folder." );
-				return false;
-			}
-
-			if(FolderPath.StartsWith( TEXT("..")))
-			{
-				OutErrorMessage = LOCTEXT( "InvalidFolderName_RelativePath", "Relative path is not accepted." );
-				return false;
-			}
-
-			if( FolderPath.StartsWith( TEXT("/Content") ) )
-			{
-				FolderPath = FolderPath.Replace( TEXT( "/Content" ), TEXT( "/Game" ) );
-			}
-
-			if ( FolderPath.Len() > FPlatformMisc::GetMaxPathLength() )
-			{
-				OutErrorMessage = FText::Format( LOCTEXT("InvalidFolderName_TooLongForCooking", "Filename is too long ({0} characters); this may interfere with cooking for consoles. Unreal filenames should be no longer than {1} characters. Filename value: {2}" ),
-					FText::AsNumber(FolderPath.Len()), FText::AsNumber(FPlatformMisc::GetMaxPathLength()), FText::FromString(FolderPath) );
-				return false;
-			}
-
-			if(FolderPath[FolderPath.Len() - 1] == L'/')
-			{
-				FolderPath.LeftInline(FolderPath.Len() - 1, false);
-			}
-
-			FString FolderName = FPaths::GetBaseFilename(FolderPath);
-
-			if(!VerifyObjectName(FolderName, OutErrorMessage))
-			{
-				return false;
-			}
-
-
-			const FString InvalidChars = INVALID_LONGPACKAGE_CHARACTERS TEXT("/[]"); // Slash and Square brackets are invalid characters for a folder name
-
-																					 // See if the name contains invalid characters.
-			FString Char;
-			for( int32 CharIdx = 0; CharIdx < FolderName.Len(); ++CharIdx )
-			{
-				Char = FolderName.Mid(CharIdx, 1);
-
-				if ( InvalidChars.Contains(*Char) )
+				if(FolderPath.StartsWith( TEXT("..")))
 				{
-					FString ReadableInvalidChars = InvalidChars;
-					ReadableInvalidChars.ReplaceInline(TEXT("\r"), TEXT(""));
-					ReadableInvalidChars.ReplaceInline(TEXT("\n"), TEXT(""));
-					ReadableInvalidChars.ReplaceInline(TEXT("\t"), TEXT(""));
-
-					OutErrorMessage = FText::Format(LOCTEXT("InvalidFolderName_InvalidCharacters", "A folder name may not contain any of the following characters: {0}"), FText::FromString(ReadableInvalidChars));
-					return false;
+					OutErrorMessage = LOCTEXT( "InvalidFolderName_RelativePath", "Relative path is not accepted." );
 				}
-			}
-
-			if(!FFileHelper::IsFilenameValidForSaving(FolderPath, OutErrorMessage))
-			{
-				return false;
-			}
+				else
+				{
+					if( FolderPath.StartsWith( TEXT("/Content") ) )
+					{
+						FolderPath = FolderPath.Replace( TEXT( "/Content" ), TEXT( "/Game" ) );
+					}
 
-			FString PathOnDisk;
-			if (!FPackageName::TryConvertLongPackageNameToFilename(FolderPath, PathOnDisk))
-			{
-				OutErrorMessage = FText::Format(LOCTEXT("RenameFolderFailedDiskPath", "Folder path could not be converted to disk path: '{0}'"), FText::FromString(FolderPath));
-				return false;
+					if(FolderPath.Len() <= FPlatformMisc::GetMaxPathLength())
+					{
+						if(FolderPath[FolderPath.Len() - 1] == L'/')
+						{
+							FolderPath.LeftInline(FolderPath.Len() - 1, false);
+						}
+
+						if(FolderPath != TEXT( "/Game" ))
+						{
+							FString FolderName = FPaths::GetBaseFilename(FolderPath);
+
+							if(VerifyObjectName(FolderName, OutErrorMessage))
+							{
+								const FString InvalidChars = INVALID_LONGPACKAGE_CHARACTERS TEXT("/[]"); // Slash and Square brackets are invalid characters for a folder name
+
+																										 // See if the name contains invalid characters.
+								FString Char;
+								for( int32 CharIdx = 0; CharIdx < FolderName.Len(); ++CharIdx )
+								{
+									Char = FolderName.Mid(CharIdx, 1);
+
+									if ( InvalidChars.Contains(*Char) )
+									{
+										FString ReadableInvalidChars = InvalidChars;
+										ReadableInvalidChars.ReplaceInline(TEXT("\r"), TEXT(""));
+										ReadableInvalidChars.ReplaceInline(TEXT("\n"), TEXT(""));
+										ReadableInvalidChars.ReplaceInline(TEXT("\t"), TEXT(""));
+
+										OutErrorMessage = FText::Format(LOCTEXT("InvalidFolderName_InvalidCharacters", "A folder name may not contain any of the following characters: {0}"), FText::FromString(ReadableInvalidChars));
+										break;
+									}
+								}
+
+								if(OutErrorMessage.IsEmpty())
+								{
+									if(FFileHelper::IsFilenameValidForSaving(FolderPath, OutErrorMessage))
+									{
+										FString PathOnDisk;
+										if(FPackageName::TryConvertLongPackageNameToFilename(FolderPath, PathOnDisk))
+										{
+											// Make sure we are not creating a folder path that is too long
+											if (PathOnDisk.Len() > FPlatformMisc::GetMaxPathLength() - 32/*MAX_CLASS_NAME_LENGTH*/)
+											{
+												// The full path for the folder is too long
+												OutErrorMessage = FText::Format(LOCTEXT("RenameFolderPathTooLong",
+													"The full path for the folder is too deep, the maximum is '{0}'. Please choose a shorter name for the folder or create it in a shallower folder structure."),
+													FText::AsNumber(FPlatformMisc::GetMaxPathLength() - 32/*MAX_CLASS_NAME_LENGTH*/));
+											}
+										}
+										else
+										{
+											OutErrorMessage = FText::Format(LOCTEXT("RenameFolderFailedDiskPath", "Folder path could not be converted to disk path: '{0}'"), FText::FromString(FolderPath));
+										}
+									}
+								}
+							}
+						}
+					}
+					else
+					{
+						OutErrorMessage = FText::Format( LOCTEXT("InvalidFolderName_TooLongForCooking", "Filename is too long ({0} characters); this may interfere with cooking for consoles. Unreal filenames should be no longer than {1} characters. Filename value: {2}" ),
+							FText::AsNumber(FolderPath.Len()), FText::AsNumber(FPlatformMisc::GetMaxPathLength()), FText::FromString(FolderPath) );
+					}
+				}
 			}
-
-			// Make sure we are not creating a folder path that is too long
-			if (PathOnDisk.Len() > FPlatformMisc::GetMaxPathLength() - 32/*MAX_CLASS_NAME_LENGTH*/)
+			else
 			{
-				// The full path for the folder is too long
-				OutErrorMessage = FText::Format(LOCTEXT("RenameFolderPathTooLong",
-					"The full path for the folder is too deep, the maximum is '{0}'. Please choose a shorter name for the folder or create it in a shallower folder structure."),
-					FText::AsNumber(FPlatformMisc::GetMaxPathLength() - 32/*MAX_CLASS_NAME_LENGTH*/));
-				// Return false to indicate that the user should enter a new name for the folder
-				return false;
+				OutErrorMessage = LOCTEXT( "InvalidFolderName_IsTooShort", "Please provide a name for this folder." );
 			}
 
-			return true;
+			ContentFolderTextBox->SetError(OutErrorMessage);
 		}
 
 		void UpdateContentFolderText()
 		{
+			ContentFolderTextBox->SetError(FText::GetEmpty());
+
 			if(UDataprepContentConsumer* Consumer = ConsumerPtr.Get())
 			{
 				FString TargetContentFolder( Consumer->GetTargetContentFolder() );
@@ -302,10 +308,11 @@ namespace DatasmithConsumerDetailsUtil
 					{
 						Transaction.Cancel();
 						UE_LOG( LogDatasmithImport, Error, TEXT("%s"), *ErrorReason.ToString() );
-						UpdateContentFolderText();
 					}
 				}
 
+				UpdateContentFolderText();
+
 				bProcessing = false;
 			}
 		}
@@ -345,7 +352,7 @@ namespace DatasmithConsumerDetailsUtil
 					.HintText(LOCTEXT("DataprepLevelProperty_HintText", "Set the name of the level to save in"))
 					.IsReadOnly(false)
 					.OnTextCommitted(FOnTextCommitted::CreateSP(this, &SLevelProperty::OnTextCommitted))
-					.OnVerifyTextChanged(FOnVerifyTextChanged::CreateSP(this, &SLevelProperty::OnVerifyText))
+					.OnTextChanged(FOnTextChanged::CreateSP(this, &SLevelProperty::OnTextChanged))
 				]
 			];
 
@@ -355,6 +362,8 @@ namespace DatasmithConsumerDetailsUtil
 	private:
 		void UpdateLevelText()
 		{
+			LevelTextBox->SetError(FText::GetEmpty());
+
 			if(UDatasmithConsumer* Consumer = ConsumerPtr.Get())
 			{
 				LevelTextBox->SetText( FText::FromString( Consumer->GetLevelName() ) );
@@ -365,20 +374,23 @@ namespace DatasmithConsumerDetailsUtil
 			}
 		}
 
-		bool OnVerifyText(const FText& InText, FText& OutErrorMessage)
+		void OnTextChanged(const FText& InText)
 		{
 			const FString NewLevelName = InText.ToString();
 
+			FText OutErrorMessage = FText::GetEmpty();
+
 			int32 Index;
 			if(NewLevelName.FindChar(L'/', Index))
 			{
 				OutErrorMessage = LOCTEXT( "InvalidLevelName_RelativePath", "Path or relative path is not accepted." );
-				bValidText = false;
+			}
+			else
+			{
+				VerifyObjectName(NewLevelName, OutErrorMessage);
 			}
 
-			bValidText = VerifyObjectName(NewLevelName, OutErrorMessage);
-
-			return bValidText;
+			LevelTextBox->SetError(OutErrorMessage);
 		}
 
 		void OnTextCommitted( const FText& NewText, ETextCommit::Type CommitType)
@@ -388,17 +400,11 @@ namespace DatasmithConsumerDetailsUtil
 				return;
 			}
 
-			if(!bValidText)
-			{
-				UpdateLevelText();
-				bValidText = true;
-				return;
-			}
-
 			if( UDatasmithConsumer* DataprepConsumer = ConsumerPtr.Get() )
 			{
 				bProcessing = true;
-				FString NewLevelName( NewText.ToString() );
+
+				const FString NewLevelName = NewText.ToString();
 
 				if(NewLevelName != DataprepConsumer->GetLevelName())
 				{
@@ -414,6 +420,7 @@ namespace DatasmithConsumerDetailsUtil
 						UE_LOG( LogDatasmithImport, Error, TEXT("Cannot create a level named %s - %s"), *NewLevelName, *OutReason.ToString() );
 					}
 				}
+
 				bProcessing = false;
 			}
 		}
