Fixed flawed logic when validating folder and level names in widget

#rnx

diff --git a/Engine/Plugins/Enterprise/DatasmithImporter/Source/DatasmithImporter/Private/UI/DatasmithConsumerDetails.cpp b/Engine/Plugins/Enterprise/DatasmithImporter/Source/DatasmithImporter/Private/UI/DatasmithConsumerDetails.cpp
--- a/Engine/Plugins/Enterprise/DatasmithImporter/Source/DatasmithImporter/Private/UI/DatasmithConsumerDetails.cpp
+++ b/Engine/Plugins/Enterprise/DatasmithImporter/Source/DatasmithImporter/Private/UI/DatasmithConsumerDetails.cpp
@@ -82,7 +82,8 @@ namespace DatasmithConsumerDetailsUtil
 		void Construct(const FArguments& InArgs, UDatasmithConsumer* InConsumer)
 		{
 			ConsumerPtr = InConsumer;
-			bProcessing = false;
+			bCommitting = false;
+			bTextValid = true;
 
 			FSlateFontInfo FontInfo = IDetailLayoutBuilder::GetDetailFont();
 
@@ -245,6 +246,8 @@ namespace DatasmithConsumerDetailsUtil
 				OutErrorMessage = LOCTEXT( "InvalidFolderName_IsTooShort", "Please provide a name for this folder." );
 			}
 
+			bTextValid = OutErrorMessage.IsEmpty();
+
 			ContentFolderTextBox->SetError(OutErrorMessage);
 		}
 
@@ -275,14 +278,21 @@ namespace DatasmithConsumerDetailsUtil
 
 		void OnTextCommitted( const FText& NewText, ETextCommit::Type CommitType)
 		{
-			if(bProcessing)
+			if(bCommitting)
+			{
+				return;
+			}
+
+			if(!bTextValid)
 			{
+				UpdateContentFolderText();
+				bTextValid = true;
 				return;
 			}
 
 			if( UDataprepContentConsumer* DataprepConsumer = ConsumerPtr.Get() )
 			{
-				bProcessing = true;
+				bCommitting = true;
 
 				FString NewContentFolder( NewText.ToString() );
 
@@ -309,19 +319,21 @@ namespace DatasmithConsumerDetailsUtil
 						Transaction.Cancel();
 						UE_LOG( LogDatasmithImport, Error, TEXT("%s"), *ErrorReason.ToString() );
 					}
-				}
 
-				UpdateContentFolderText();
+					UpdateContentFolderText();
+				}
 
-				bProcessing = false;
+				bCommitting = false;
 			}
 		}
 
 	private:
 		TWeakObjectPtr< UDatasmithConsumer > ConsumerPtr;
 		TSharedPtr< SEditableTextBox > ContentFolderTextBox;
-		// Boolean used to avoid re-entering UI event processing
-		bool bProcessing;
+		// Indicates committing is in progress
+		bool bCommitting;
+		// Indicates if text is valid or not
+		bool bTextValid;
 	};
 
 	class SLevelProperty : public SCompoundWidget
@@ -334,8 +346,8 @@ namespace DatasmithConsumerDetailsUtil
 		void Construct(const FArguments& InArgs, UDatasmithConsumer* InConsumer)
 		{
 			ConsumerPtr = InConsumer;
-			bProcessing = false;
-			bValidText = true;
+			bCommitting = false;
+			bTextValid = true;
 
 			FSlateFontInfo FontInfo = IDetailLayoutBuilder::GetDetailFont();
 
@@ -390,19 +402,28 @@ namespace DatasmithConsumerDetailsUtil
 				VerifyObjectName(NewLevelName, OutErrorMessage);
 			}
 
+			bTextValid = OutErrorMessage.IsEmpty();
+
 			LevelTextBox->SetError(OutErrorMessage);
 		}
 
 		void OnTextCommitted( const FText& NewText, ETextCommit::Type CommitType)
 		{
-			if(bProcessing)
+			if(bCommitting)
+			{
+				return;
+			}
+
+			if(!bTextValid)
 			{
+				UpdateLevelText();
+				bTextValid = true;
 				return;
 			}
 
 			if( UDatasmithConsumer* DataprepConsumer = ConsumerPtr.Get() )
 			{
-				bProcessing = true;
+				bCommitting = true;
 
 				const FString NewLevelName = NewText.ToString();
 
@@ -415,22 +436,23 @@ namespace DatasmithConsumerDetailsUtil
 					{
 						Transaction.Cancel();
 
-						UpdateLevelText();
-
 						UE_LOG( LogDatasmithImport, Error, TEXT("Cannot create a level named %s - %s"), *NewLevelName, *OutReason.ToString() );
 					}
+
+					UpdateLevelText();
 				}
 
-				bProcessing = false;
+				bCommitting = false;
 			}
 		}
 
 	private:
 		TWeakObjectPtr< UDatasmithConsumer > ConsumerPtr;
 		TSharedPtr< SEditableTextBox > LevelTextBox;
-		// Boolean used to avoid re-entering UI event processing
-		bool bProcessing;
-		bool bValidText;
+		// Indicates committing is in progress. Used to avoid re-entering UI event processing
+		bool bCommitting;
+		// Indicates if text in widget is valid or not
+		bool bTextValid;
 	};
 }
 
