Fixed crash when creating asset

#rnx

diff --git a/Engine/Plugins/Enterprise/DataprepEditor/Source/DataprepCore/Private/DataprepCoreUtils.cpp b/Engine/Plugins/Enterprise/DataprepEditor/Source/DataprepCore/Private/DataprepCoreUtils.cpp
--- a/Engine/Plugins/Enterprise/DataprepEditor/Source/DataprepCore/Private/DataprepCoreUtils.cpp
+++ b/Engine/Plugins/Enterprise/DataprepEditor/Source/DataprepCore/Private/DataprepCoreUtils.cpp
@@ -320,6 +320,7 @@ bool FDataprepCoreUtils::ExecuteDataprep(UDataprepAssetInterface* DataprepAssetI
 			Context.SetWorld( TransientWorld.Get() )
 				.SetAssets( Assets )
 				.SetTransientContentFolder( TransientContentFolder )
+				.SetSilentMode(true)
 				.SetLogger( Logger )
 				.SetProgressReporter( Reporter );
 
diff --git a/Engine/Plugins/Enterprise/DataprepEditor/Source/DataprepCore/Public/DataprepContentConsumer.h b/Engine/Plugins/Enterprise/DataprepEditor/Source/DataprepCore/Public/DataprepContentConsumer.h
--- a/Engine/Plugins/Enterprise/DataprepEditor/Source/DataprepCore/Public/DataprepContentConsumer.h
+++ b/Engine/Plugins/Enterprise/DataprepEditor/Source/DataprepCore/Public/DataprepContentConsumer.h
@@ -70,6 +70,12 @@ struct FDataprepConsumerContext
 		return *this;
 	}
 
+	FDataprepConsumerContext& SetSilentMode( bool InSilentMode )
+	{
+		bSilentMode = InSilentMode;
+		return *this;
+	}
+
 	/** Hold onto the world the consumer will process */
 	TWeakObjectPtr< UWorld > WorldPtr;
 
diff --git a/Engine/Plugins/Enterprise/DatasmithImporter/Source/DatasmithImporter/Private/DatasmithConsumer.cpp b/Engine/Plugins/Enterprise/DatasmithImporter/Source/DatasmithImporter/Private/DatasmithConsumer.cpp
--- a/Engine/Plugins/Enterprise/DatasmithImporter/Source/DatasmithImporter/Private/DatasmithConsumer.cpp
+++ b/Engine/Plugins/Enterprise/DatasmithImporter/Source/DatasmithImporter/Private/DatasmithConsumer.cpp
@@ -310,7 +310,7 @@ bool UDatasmithConsumer::Initialize()
 
 	ProgressTaskPtr->ReportNextStep( LOCTEXT( "DatasmithImportFactory_Initialize", "Preparing world ...") );
 
-	if(!CheckOutputDirectives())
+	if(!ValidateAssets())
 	{
 		return false;
 	}
@@ -775,7 +775,8 @@ bool UDatasmithConsumer::CanCreateLevel(const FString& RequestedFolder, const FS
 
 	if(FDatasmithImporterUtils::CanCreateAsset(AssetPathName, UWorld::StaticClass()) == FDatasmithImporterUtils::EAssetCreationStatus::CS_CanCreate)
 	{
-		if(FDatasmithImporterImpl::CheckAssetPersistenceValidity(ObjectPath.GetLongPackageName(), *ImportContextPtr, FPackageName::GetMapPackageExtension()))
+		FText OutReason;
+		if(FDatasmithImporterImpl::CheckAssetPersistenceValidity(ObjectPath.GetLongPackageName(), *ImportContextPtr, FPackageName::GetMapPackageExtension(), OutReason))
 		{
 			FString PackageFilename;
 			FPackageName::TryConvertLongPackageNameToFilename( ObjectPath.GetLongPackageName(), PackageFilename, FPackageName::GetMapPackageExtension() );
@@ -992,81 +993,207 @@ ULevel* UDatasmithConsumer::FindOrAddLevel(const FString& InLevelName)
 	return nullptr;
 }
 
-bool UDatasmithConsumer::CheckOutputDirectives()
+bool UDatasmithConsumer::ValidateAssets()
 {
-	auto CanCreateAsset = [ImportContext = ImportContextPtr.Get()](const FString& AssetPathName,const UClass* AssetClass)
+	auto CanCreateAsset = [&](const FString& AssetPathName,const UClass* AssetClass)
 	{
-		if(FDatasmithImporterUtils::CanCreateAsset(AssetPathName, AssetClass) == FDatasmithImporterUtils::EAssetCreationStatus::CS_CanCreate)
+		FText OutReason;
+		if(!FDatasmithImporterUtils::CanCreateAsset(AssetPathName, AssetClass, OutReason))
 		{
-			return FDatasmithImporterImpl::CheckAssetPersistenceValidity(FPaths::GetPath(AssetPathName), *ImportContext, FPackageName::GetAssetPackageExtension());
+			const FTextFormat TextFormat(LOCTEXT( "DatasmithConsumer_CannotCreateAsset", "Cannot create asset {0}. {1}" ));
+			const FText Message = FText::Format( TextFormat, FText::FromString(AssetPathName), OutReason );
+			LogError(Message);
+
+			return false;
 		}
 
-		return false;
+		return true;
 	};
-
-	const bool bShowDialog = !Context.bSilentMode && !IsRunningCommandlet();
-
+	
 	// Collect garbage to clear out the destroyed level
 	CollectGarbage( GARBAGE_COLLECTION_KEEPFLAGS );
 
-	bool bCannotCreateAsset = false;
+	TArray< TWeakObjectPtr< UObject > > Assets = MoveTemp(Context.Assets);
+	Context.Assets.Reserve(Assets.Num());
+
+	// Map holding requested package path and actual package path
+	TMap<UObject*, TPair<FString, FString> > AssetPackageInfoMap;
+	AssetPackageInfoMap.Reserve(Assets.Num());
+
+	TSet< FString > AssetsToCreate;
+	AssetsToCreate.Reserve(Assets.Num());
 
-	TSet<FString> AssetPaths;
-	AssetPaths.Reserve(Context.Assets.Num());
+	TSet< UObject* > AssetsWithIssues;
+	AssetsWithIssues.Reserve(Assets.Num());
 
-	for(const TWeakObjectPtr< UObject >& AssetPtr : Context.Assets)
+	for(const TWeakObjectPtr< UObject >& AssetPtr : Assets)
 	{
 		if(UObject* Asset = AssetPtr.Get())
 		{
+			TPair<FString, FString>& AssetPackageInfo = AssetPackageInfoMap.Add(Asset);
+
+			bool bAssetWithMarker = false;
+
 			const FString& OutputFolder = DatasmithConsumerUtils::GetMarker(Asset, UDataprepContentConsumer::RelativeOutput);
-			if(OutputFolder.Len() > 0)
+			if(!OutputFolder.IsEmpty())
 			{
-				FString AssetName = Asset->GetName();
-				FSoftObjectPath AssetSoftObjectPath(FPaths::Combine(TargetContentFolder, OutputFolder, AssetName) + "." + AssetName);
+				const FString AssetName = Asset->GetName();
+				const FSoftObjectPath AssetSoftObjectPath(FPaths::Combine(TargetContentFolder, OutputFolder, AssetName) + "." + AssetName);
+
+				AssetPackageInfo.Key = AssetSoftObjectPath.GetLongPackageName();
 
-				if(Asset->GetPathName() != AssetSoftObjectPath.GetLongPackageName())
+				// If asset can be created in memory, do nothing
+				if(!CanCreateAsset(AssetSoftObjectPath.GetAssetPathString(), Asset->GetClass()))
 				{
-					if(!CanCreateAsset(AssetSoftObjectPath.GetAssetPathString(), Asset->GetClass()))
-					{
-						const FTextFormat TextFormat(LOCTEXT( "DatasmithConsumer_CannotCreateAsset", "Cannot create asset {0}. Commit will be aborted." ));
-						const FText Message = FText::Format( TextFormat, FText::FromString(AssetSoftObjectPath.GetAssetPathString()) );
-						LogError(Message);
+				}
+				// Verify there is no collision with other assets to be moved
+				else if(AssetsToCreate.Contains(AssetPackageInfo.Key))
+				{
+					const FTextFormat TextFormat(LOCTEXT( "DatasmithConsumer_DuplicateAsset", "Cannot create asset {0}. Another asset with the same name will be created in the same folder {1}." ));
+					const FText Message = FText::Format( TextFormat, FText::FromString(AssetSoftObjectPath.GetAssetPathString()), FText::FromString(FPaths::Combine(TargetContentFolder, OutputFolder)) );
+					LogError(Message);
+				}
+				else
+				{
+					// Asset could be created in memory using output folder directive
+					AssetsToCreate.Add(AssetPackageInfo.Key);
+					AssetPackageInfo.Value = AssetPackageInfo.Key;
+					bAssetWithMarker = true;
+				}
+			}
 
-						bCannotCreateAsset = true;
+			// Check if asset can be created in memory using policy based on target content folder
+			if(AssetPackageInfo.Value.IsEmpty())
+			{
+				const FSoftObjectPath AssetSoftObjectPath(Asset);
+				const FString AssetPath = AssetSoftObjectPath.GetAssetPathString().Replace(*Context.TransientContentFolder, *TargetContentFolder);
+				const FString LongPackageName = AssetSoftObjectPath.GetLongPackageName().Replace(*Context.TransientContentFolder, *TargetContentFolder);
+
+				if(!CanCreateAsset(AssetPath, Asset->GetClass()))
+				{
+					// Asset without output folder directive cannot be saved
+					if(AssetPackageInfo.Key.IsEmpty())
+					{
+						AssetPackageInfo.Key = LongPackageName;
 					}
-					else if(AssetPaths.Contains(AssetSoftObjectPath.GetLongPackageName()))
+
+					AssetsWithIssues.Add(Asset);
+				}
+				// Asset with output folder directive could be saved using regular policy
+				else if(!AssetPackageInfo.Key.IsEmpty())
+				{
+					DatasmithConsumerUtils::SetMarker(Asset, UDataprepContentConsumer::RelativeOutput, TEXT(""));
+					AssetPackageInfo.Value = LongPackageName;
+					AssetsWithIssues.Add(Asset);
+				}
+				else
+				{
+					AssetPackageInfo.Key = LongPackageName;
+					AssetPackageInfo.Value = AssetPackageInfo.Key;
+				}
+			}
+
+			// Verify asset can be saved to disk
+			if(!AssetPackageInfo.Value.IsEmpty())
+			{
+				FText OutReason = FText::GetEmpty();
+
+				// If asset cannot be saved to disk using output folder, check if it can do so using regular folder policy
+				if(!FDatasmithImporterImpl::CheckAssetPersistenceValidity(AssetPackageInfo.Value, *ImportContextPtr, FPackageName::GetAssetPackageExtension(), OutReason) && bAssetWithMarker)
+				{
+					LogWarning(OutReason);
+
+					// Overwrite output folder directive
+					DatasmithConsumerUtils::SetMarker(Asset, UDataprepContentConsumer::RelativeOutput, TEXT(""));
+
+					FSoftObjectPath AssetSoftObjectPath(Asset);
+					const FString AssetPath = AssetSoftObjectPath.GetAssetPathString().Replace(*Context.TransientContentFolder, *TargetContentFolder);
+
+					const FTextFormat Format(LOCTEXT("DatasmithConsumer_OutputFolderFailed", "Cannot save {0} in output folder, trying to save as {1}."));
+					OutReason = FText::Format( Format, FText::FromString(Asset->GetName()), FText::FromString(FPaths::GetBaseFilename(AssetPath, false)));
+					LogWarning(OutReason);
+
+					if(CanCreateAsset(AssetPath, Asset->GetClass()))
 					{
-						const FTextFormat TextFormat(LOCTEXT( "DatasmithConsumer_DuplicateAsset", "Cannot create asset {0}. Another asset with the same name will be created in the same folder {1}. Commit will be aborted." ));
-						const FText Message = FText::Format( TextFormat, FText::FromString(AssetSoftObjectPath.GetAssetPathString()), FText::FromString(FPaths::Combine(TargetContentFolder, OutputFolder)) );
-						LogError(Message);
+						const FString LongPackageName = AssetSoftObjectPath.GetLongPackageName().Replace(*Context.TransientContentFolder, *TargetContentFolder);
+						OutReason = FText::GetEmpty();
+
+						if(FDatasmithImporterImpl::CheckAssetPersistenceValidity(LongPackageName, *ImportContextPtr, FPackageName::GetAssetPackageExtension(), OutReason))
+						{
+							AssetPackageInfo.Value = LongPackageName;
+							Context.Assets.Emplace(Asset);
+
+							AssetsWithIssues.Add(Asset);
+						}
+						else
+						{
+							LogWarning(OutReason);
 
-						bCannotCreateAsset = true;
+							AssetPackageInfo.Value = FString();
+							AssetsWithIssues.Add(Asset);
+						}
 					}
 					else
 					{
-						AssetPaths.Add(AssetSoftObjectPath.GetLongPackageName());
+						AssetPackageInfo.Value = FString();
+						AssetsWithIssues.Add(Asset);
 					}
 				}
+				else
+				{
+					Context.Assets.Emplace(Asset);
+				}
 			}
 		}
 	}
 
-	if(bCannotCreateAsset)
+	if(AssetsWithIssues.Num() > 0)
 	{
+		const bool bShowDialog = !Context.bSilentMode && !IsRunningCommandlet();
 
-		const FText Message = LOCTEXT( "DatasmithConsumer_CreateAbortCommit", "Cannot proceed with commit because some assets and/or levels cannot be created.\nCheck your log for details, fix all issues and commit again" );
 		
 		if(bShowDialog)
 		{
-			const FText Title( LOCTEXT( "DatasmithConsumer_CreateAbortCommitTitle", "Cannot create some assets" ) );
-			FMessageDialog::Open( EAppMsgType::Ok, Message, &Title );
+			FString AssetsToBeSkippedString;
+			FString AssetsNotMovedString;
+			for(UObject* Asset : AssetsWithIssues)
+			{
+				const TPair<FString, FString>& AssetPackageInfo = AssetPackageInfoMap[Asset];
+				if(AssetPackageInfo.Value.IsEmpty())
+				{
+					AssetsToBeSkippedString.Append(FString::Printf(TEXT("\t%s\n"), *AssetPackageInfo.Key));
+				}
+				else
+				{
+					AssetsNotMovedString.Append(FString::Printf(TEXT("\t%s will be created as %s\n"), *AssetPackageInfo.Key, *AssetPackageInfo.Value));
+				}
+			}
+
+			FText AssetsToBeSkippedText;
+			if(!AssetsToBeSkippedString.IsEmpty())
+			{
+				AssetsToBeSkippedText = FText::Format( LOCTEXT( "DatasmithConsumer_AssetsToBeSkipped", "The follwing assets will not be saved:\n{0}\n"), FText::FromString(AssetsToBeSkippedString) );
+			}
+
+			FText AssetsNotMovedText;
+			if(!AssetsNotMovedString.IsEmpty())
+			{
+				AssetsNotMovedText = FText::Format( LOCTEXT( "DatasmithConsumer_AssetsNotMoved", "The follwing assets will not be saved in their output folder:\n{0}\n"), FText::FromString(AssetsNotMovedString) );
+			}
+
+			const FText Title( LOCTEXT( "DatasmithConsumer_SavingIssues", "Some assets may not be saved properly..." ) );
+			const FText Message = FText::Format( LOCTEXT( "DatasmithConsumer_AssetsWithIssues", "All assets cannot be created in their destination folder.\nBelow is the list of assets with issues. See output log for details.\nClick \'Yes\' to continue with the commit.\n\n{0}{1}\n" ), AssetsToBeSkippedText, AssetsNotMovedText);
+
+			if(FMessageDialog::Open(EAppMsgType::YesNo, Message, &Title) != EAppReturnType::Yes)
+			{
+				return false;
+			}
 		}
 		else
 		{
-			LogError(Message);
+			const FText Message = LOCTEXT( "DatasmithConsumer_CommitWithErrors", "All assets could not be saved in their destination folder. Committing anyway. Check your log for details." );
+			LogWarning(Message);
 		}
-
-		return false;
 	}
 
 	return true;
diff --git a/Engine/Plugins/Enterprise/DatasmithImporter/Source/DatasmithImporter/Private/DatasmithConsumer.h b/Engine/Plugins/Enterprise/DatasmithImporter/Source/DatasmithImporter/Private/DatasmithConsumer.h
--- a/Engine/Plugins/Enterprise/DatasmithImporter/Source/DatasmithImporter/Private/DatasmithConsumer.h
+++ b/Engine/Plugins/Enterprise/DatasmithImporter/Source/DatasmithImporter/Private/DatasmithConsumer.h
@@ -83,7 +83,7 @@ private:
 	 * @return True if no issue or user has accepted to continue with execution of consumer
 	 * If in command mode, the default behavior is to continue.
 	 */
-	bool CheckOutputDirectives();
+	bool ValidateAssets();
 
 	/** Create the world which will be used to execute the consumer */
 	bool CreateWorld();
diff --git a/Engine/Plugins/Enterprise/DatasmithImporter/Source/DatasmithImporter/Private/Utility/DatasmithImporterImpl.cpp b/Engine/Plugins/Enterprise/DatasmithImporter/Source/DatasmithImporter/Private/Utility/DatasmithImporterImpl.cpp
--- a/Engine/Plugins/Enterprise/DatasmithImporter/Source/DatasmithImporter/Private/Utility/DatasmithImporterImpl.cpp
+++ b/Engine/Plugins/Enterprise/DatasmithImporter/Source/DatasmithImporter/Private/Utility/DatasmithImporterImpl.cpp
@@ -837,6 +837,18 @@ FDatasmithImporterImpl::FScopedFinalizeActorChanges::~FScopedFinalizeActorChange
 }
 
 bool FDatasmithImporterImpl::CheckAssetPersistenceValidity(const FString& PackageName, FDatasmithImportContext& ImportContext, const FString& Extension)
+{
+	FText OutReason;
+	if(!CheckAssetPersistenceValidity(PackageName, ImportContext, Extension, OutReason))
+	{
+		ImportContext.LogWarning(OutReason);
+		return false;
+	}
+
+	return true;
+}
+
+bool FDatasmithImporterImpl::CheckAssetPersistenceValidity(const FString& PackageName, FDatasmithImportContext& ImportContext, const FString& Extension, FText& OutReason)
 {
 	// Check that package can be saved
 	const FString BasePackageFileName = FPackageName::LongPackageNameToFilename( PackageName );
@@ -845,13 +857,12 @@ bool FDatasmithImporterImpl::CheckAssetPersistenceValidity(const FString& Packag
 	// Create fake filename of same length of final asset file name to test ability to write
 	const FString FakeAbsolutePathToAsset = AbsolutePathToAsset + Extension;
 
-	bool bValidPersistence = true;
+	OutReason = FText::GetEmpty();
 
 	// Verify asset file name does not exceed OS' maximum path length
 	if( FPlatformMisc::GetMaxPathLength() < FakeAbsolutePathToAsset.Len() )
 	{
-		ImportContext.LogWarning(FText::Format(LOCTEXT("DatasmithImportInvalidLength", "Saving may partially fail because path for asset {0} is too long. Rename before saving."), FText::FromString( PackageName )));
-		bValidPersistence = false;
+		OutReason = FText::Format(LOCTEXT("DatasmithImportInvalidLength", "Saving may partially fail because path for asset {0} is too long. Rename before saving."), FText::FromString( PackageName ));
 	}
 	// Verify user can overwrite existing file
 	else if( IFileManager::Get().FileExists( *FakeAbsolutePathToAsset ) )
@@ -860,7 +871,7 @@ bool FDatasmithImporterImpl::CheckAssetPersistenceValidity(const FString& Packag
 		if ( FileStatData.bIsReadOnly )
 		{
 			// Check to see if the file is not under source control
-			bool bWarnUser = true;
+			bool bCanCheckedOut = false;
 
 			ISourceControlProvider& SourceControlProvider = ISourceControlModule::Get().GetProvider();
 			if (SourceControlProvider.IsAvailable() && SourceControlProvider.IsEnabled())
@@ -870,14 +881,13 @@ bool FDatasmithImporterImpl::CheckAssetPersistenceValidity(const FString& Packag
 				if( SourceControlState.IsValid() && SourceControlState->CanCheckout() )
 				{
 					// User will be prompted to check out this file when he/she saves the asset. No need to warn.
-					bWarnUser = false;
+					bCanCheckedOut = true;
 				}
 			}
 
-			if(bWarnUser)
+			if(!bCanCheckedOut)
 			{
-				ImportContext.LogWarning(FText::Format(LOCTEXT("DatasmithImportInvalidSaving", "Saving may partially fail because file asset {0} cannot be overwritten. Check your privileges."), FText::FromString( PackageName )));
-				bValidPersistence = false;
+				OutReason = FText::Format(LOCTEXT("DatasmithImportInvalidSaving", "Saving may partially fail because file asset {0} cannot be overwritten. Check your privileges."), FText::FromString( PackageName ));
 			}
 		}
 	}
@@ -904,8 +914,7 @@ bool FDatasmithImporterImpl::CheckAssetPersistenceValidity(const FString& Packag
 
 		if(WarningState != EWriteDisallowedWarningState::WDWS_WarningUnnecessary)
 		{
-			ImportContext.LogWarning(FText::Format(LOCTEXT("DatasmithImportInvalidFolder", "Cannot write in folder {0} to store asset {1}. Check access to folder."), FText::FromString( FPaths::GetPath( FakeAbsolutePathToAsset ) ), FText::FromString( PackageName )));
-			bValidPersistence = false;
+			OutReason = FText::Format(LOCTEXT("DatasmithImportInvalidFolder", "Cannot write in folder {0} to store asset {1}. Check access to folder."), FText::FromString( FPaths::GetPath( FakeAbsolutePathToAsset ) ), FText::FromString( PackageName ));
 		}
 	}
 
@@ -938,11 +947,10 @@ bool FDatasmithImporterImpl::CheckAssetPersistenceValidity(const FString& Packag
 
 	if (AbsoluteCookPathToAsset.Len() > MaxCookPath)
 	{
-		ImportContext.LogWarning(FText::Format(LOCTEXT("DatasmithImportInvalidCooking", "Cooking may fail because path for asset {0} is too long. Rename before cooking."), FText::FromString(PackageName)));
-		bValidPersistence = false;
+		OutReason = FText::Format(LOCTEXT("DatasmithImportInvalidCooking", "Cooking may fail because path for asset {0} is too long. Rename before cooking."), FText::FromString(PackageName));
 	}
 
-	return bValidPersistence;
+	return OutReason.IsEmpty();
 }
 
 #undef LOCTEXT_NAMESPACE
diff --git a/Engine/Plugins/Enterprise/DatasmithImporter/Source/DatasmithImporter/Private/Utility/DatasmithImporterImpl.h b/Engine/Plugins/Enterprise/DatasmithImporter/Source/DatasmithImporter/Private/Utility/DatasmithImporterImpl.h
--- a/Engine/Plugins/Enterprise/DatasmithImporter/Source/DatasmithImporter/Private/Utility/DatasmithImporterImpl.h
+++ b/Engine/Plugins/Enterprise/DatasmithImporter/Source/DatasmithImporter/Private/Utility/DatasmithImporterImpl.h
@@ -65,6 +65,7 @@ public:
 	 *			- ContentBrowserUtils::IsValidPackageForCooking
 	 * #ueent_todo Make ContentBrowserUtils::IsValidPackageForCooking public. Not game to do it so late in 4.22
 	 */
+	static bool DATASMITHIMPORTER_API CheckAssetPersistenceValidity(const FString& PackageName, FDatasmithImportContext& ImportContext, const FString& Extension, FText& OutReason);
 	static bool DATASMITHIMPORTER_API CheckAssetPersistenceValidity(const FString& PackageName, FDatasmithImportContext& ImportContext, const FString& Extension);
 
 	static void CheckAssetPersistenceValidity(UObject* Asset, FDatasmithImportContext& ImportContext);
