Fixed broken parenting when re-committing a Dataprep asset

#rnx

diff --git a/Engine/Plugins/Enterprise/DatasmithImporter/Source/DatasmithImporter/Private/DatasmithConsumer.cpp b/Engine/Plugins/Enterprise/DatasmithImporter/Source/DatasmithImporter/Private/DatasmithConsumer.cpp
--- a/Engine/Plugins/Enterprise/DatasmithImporter/Source/DatasmithImporter/Private/DatasmithConsumer.cpp
+++ b/Engine/Plugins/Enterprise/DatasmithImporter/Source/DatasmithImporter/Private/DatasmithConsumer.cpp
@@ -1376,13 +1376,21 @@ namespace DatasmithConsumerUtils
 			{
 				if(IInterface_AssetUserData* TargetAssetUserDataInterface = Cast< IInterface_AssetUserData >(ActorRootComponent))
 				{
-					if(UAssetUserData* SourceDatasmithUserData = SourceAssetUserDataInterface->GetAssetUserDataOfClass(UDatasmithAssetUserData::StaticClass()))
+					if(UDatasmithAssetUserData* SourceDatasmithUserData = SourceAssetUserDataInterface->GetAssetUserData<UDatasmithAssetUserData>())
 					{
-						UAssetUserData* TargetDatasmithUserData = DuplicateObject<UAssetUserData>(SourceDatasmithUserData, ActorRootComponent);
-						TargetAssetUserDataInterface->AddAssetUserData(TargetDatasmithUserData);
+						if(UDatasmithAssetUserData* TargetDatasmithUserData = TargetAssetUserDataInterface->GetAssetUserData<UDatasmithAssetUserData>())
+						{
+							TargetDatasmithUserData->MetaData.Append(SourceDatasmithUserData->MetaData);
+							TargetDatasmithUserData->ObjectTemplates.Append(SourceDatasmithUserData->ObjectTemplates);
+						}
+						else
+						{
+							TargetDatasmithUserData = DuplicateObject<UDatasmithAssetUserData>(SourceDatasmithUserData, ActorRootComponent);
+							TargetAssetUserDataInterface->AddAssetUserData(TargetDatasmithUserData);
+						}
 					}
 
-					if(UAssetUserData* SourceConsumerUserData = SourceAssetUserDataInterface->GetAssetUserDataOfClass(UDataprepConsumerUserData::StaticClass()))
+					if(UAssetUserData* SourceConsumerUserData = SourceAssetUserDataInterface->GetAssetUserData<UDataprepConsumerUserData>())
 					{
 						UAssetUserData* TargetConsumerUserData = DuplicateObject<UAssetUserData>(SourceConsumerUserData, ActorRootComponent);
 						TargetAssetUserDataInterface->AddAssetUserData(TargetConsumerUserData);
