Fixed crash when output level already exists but is not loaded

#rnx

diff --git a/Engine/Plugins/Enterprise/DatasmithImporter/Source/DatasmithImporter/Private/DatasmithConsumer.cpp b/Engine/Plugins/Enterprise/DatasmithImporter/Source/DatasmithImporter/Private/DatasmithConsumer.cpp
--- a/Engine/Plugins/Enterprise/DatasmithImporter/Source/DatasmithImporter/Private/DatasmithConsumer.cpp
+++ b/Engine/Plugins/Enterprise/DatasmithImporter/Source/DatasmithImporter/Private/DatasmithConsumer.cpp
@@ -943,12 +943,14 @@ ULevel* UDatasmithConsumer::FindOrAddLevel(const FString& InLevelName)
 
 	ULevel* CurrentLevel = WorkingWorld->PersistentLevel;
 
-	// This level has not been added yet
-	FString PackageFilename;
+	// The level is not part of the world
+	FString PackageFilename;yet
 	FPackageName::TryConvertLongPackageNameToFilename( LevelPackageName, PackageFilename, FPackageName::GetMapPackageExtension() );
+	FSoftObjectPath LevelSoftObjectPath(LevelPackageName + TEXT(".") + InLevelName);
 
 	ULevelStreaming* StreamingLevel = nullptr;
-	if(FPaths::FileExists(PackageFilename))
+	// If it already exists on disk or in memory, add it to the working world
+	if(FPaths::FileExists(PackageFilename) || LevelSoftObjectPath.TryLoad())
 	{
 		FTransform LevelTransform;
 		StreamingLevel = UEditorLevelUtils::AddLevelToWorld(WorkingWorld.Get(), *LevelPackageName, ULevelStreamingAlwaysLoaded::StaticClass(), LevelTransform);
@@ -963,6 +965,7 @@ ULevel* UDatasmithConsumer::FindOrAddLevel(const FString& InLevelName)
 			ensure(false);
 		}
 	}
+	// If it does not exist, create a new one and add it to the working world
 	else
 	{
 		StreamingLevel = EditorLevelUtils::CreateNewStreamingLevelForWorld( *WorkingWorld, ULevelStreamingAlwaysLoaded::StaticClass(), *PackageFilename );
