Remove unvoluntary extra characters introducing a build break

#rnx

diff --git a/Engine/Plugins/Enterprise/DatasmithImporter/Source/DatasmithImporter/Private/DatasmithConsumer.cpp b/Engine/Plugins/Enterprise/DatasmithImporter/Source/DatasmithImporter/Private/DatasmithConsumer.cpp
--- a/Engine/Plugins/Enterprise/DatasmithImporter/Source/DatasmithImporter/Private/DatasmithConsumer.cpp
+++ b/Engine/Plugins/Enterprise/DatasmithImporter/Source/DatasmithImporter/Private/DatasmithConsumer.cpp
@@ -944,7 +944,7 @@ ULevel* UDatasmithConsumer::FindOrAddLevel(const FString& InLevelName)
 	ULevel* CurrentLevel = WorkingWorld->PersistentLevel;
 
 	// The level is not part of the world
-	FString PackageFilename;yet
+	FString PackageFilename;
 	FPackageName::TryConvertLongPackageNameToFilename( LevelPackageName, PackageFilename, FPackageName::GetMapPackageExtension() );
 	FSoftObjectPath LevelSoftObjectPath(LevelPackageName + TEXT(".") + InLevelName);
 
