Reverting changes to DatasmithConsumer

#rnx

diff --git a/Engine/Plugins/Enterprise/DatasmithImporter/Source/DatasmithImporter/Private/DatasmithConsumer.cpp b/Engine/Plugins/Enterprise/DatasmithImporter/Source/DatasmithImporter/Private/DatasmithConsumer.cpp
--- a/Engine/Plugins/Enterprise/DatasmithImporter/Source/DatasmithImporter/Private/DatasmithConsumer.cpp
+++ b/Engine/Plugins/Enterprise/DatasmithImporter/Source/DatasmithImporter/Private/DatasmithConsumer.cpp
@@ -946,11 +946,9 @@ ULevel* UDatasmithConsumer::FindOrAddLevel(const FString& InLevelName)
 	// The level is not part of the world
 	FString PackageFilename;
 	FPackageName::TryConvertLongPackageNameToFilename( LevelPackageName, PackageFilename, FPackageName::GetMapPackageExtension() );
-	FSoftObjectPath LevelSoftObjectPath(LevelPackageName + TEXT(".") + InLevelName);
 
 	ULevelStreaming* StreamingLevel = nullptr;
-	// If it already exists on disk or in memory, add it to the working world
-	if(FPaths::FileExists(PackageFilename) || LevelSoftObjectPath.TryLoad())
+	if(FPaths::FileExists(PackageFilename))
 	{
 		FTransform LevelTransform;
 		StreamingLevel = UEditorLevelUtils::AddLevelToWorld(WorkingWorld.Get(), *LevelPackageName, ULevelStreamingAlwaysLoaded::StaticClass(), LevelTransform);
