ImportContext: split Init method to provide more flexibility

#rnx

diff --git a/Engine/Plugins/Enterprise/DatasmithImporter/Source/DatasmithImporter/Private/DatasmithImportContext.cpp b/Engine/Plugins/Enterprise/DatasmithImporter/Source/DatasmithImporter/Private/DatasmithImportContext.cpp
--- a/Engine/Plugins/Enterprise/DatasmithImporter/Source/DatasmithImporter/Private/DatasmithImportContext.cpp
+++ b/Engine/Plugins/Enterprise/DatasmithImporter/Source/DatasmithImporter/Private/DatasmithImportContext.cpp
@@ -184,12 +184,6 @@ FDatasmithImportContext::FDatasmithImportContext(const FString& FileName, bool b
 	// Force the SceneHandling to be on current level by default.
 	// Note: This is done because this option was previously persisted and can get overwritten
 	Options->BaseOptions.SceneHandling = EDatasmithImportScene::CurrentLevel;
-
-	SetFileName(FileName);
-}
-
-void FDatasmithImportContext::SetFileName(const FString& FileName)
-{
 	Options->FileName = FPaths::GetCleanFilename(FileName);
 	Options->FilePath = FPaths::ConvertRelativePathToFull(FileName);
 
@@ -228,14 +222,21 @@ TSharedRef<FTokenizedMessage> FDatasmithImportContext::LogInfo(const FText& InIn
 	return Logger.Push(EMessageSeverity::Info, InInfoMessage);
 }
 
-bool FDatasmithImportContext::Init(const FString& InFileName, TSharedRef< IDatasmithScene > InScene, const FString& InImportPath, EObjectFlags InFlags, FFeedbackContext* InWarn, const TSharedPtr<FJsonObject>& ImportSettingsJson, bool bSilent)
+bool FDatasmithImportContext::Init(TSharedRef< IDatasmithScene > InScene, const FString& InImportPath, EObjectFlags InFlags, FFeedbackContext* InWarn, const TSharedPtr<FJsonObject>& ImportSettingsJson, bool bSilent)
 {
-	SetFileName(InFileName);
-	return Init(InScene, InImportPath, InFlags, InWarn, ImportSettingsJson, bSilent);
+	return InitOptions(InScene, ImportSettingsJson, bSilent)
+		&& SetupDestination(InImportPath, InFlags, InWarn, bSilent);
 }
 
-bool FDatasmithImportContext::Init(TSharedRef< IDatasmithScene > InScene, const FString& InImportPath, EObjectFlags InFlags, FFeedbackContext* InWarn, const TSharedPtr<FJsonObject>& ImportSettingsJson, bool bSilent)
+bool FDatasmithImportContext::InitOptions(TSharedRef< IDatasmithScene > InScene, const TSharedPtr<FJsonObject>& ImportSettingsJson, bool bSilent)
 {
+	if (!FModuleManager::Get().IsModuleLoaded("AssetTools"))
+	{
+		UE_LOG(LogDatasmithImport, Warning, TEXT("Import failed. The AssetTools module can't be loaded."));
+		return false;
+	}
+	Scene = InScene;
+
 	check(Options->FileName.Len() > 0);
 	check(Options->FilePath.Len() > 0);
 
@@ -248,14 +249,6 @@ bool FDatasmithImportContext::Init(TSharedRef< IDatasmithScene > InScene, const
 		ImportOptions.Add(AdditionalOption.Get());
 	}
 
-	if (!FModuleManager::Get().IsModuleLoaded("AssetTools"))
-	{
-		UE_LOG(LogDatasmithImport, Warning, TEXT("Import failed. The AssetTools module can't be loaded."));
-		return false;
-	}
-
-	Options->BaseOptions.AssetOptions.PackagePath = FName( *InImportPath );
-
 	if (bSilent)
 	{
 		// Update options based on given JSON object
@@ -268,11 +261,11 @@ bool FDatasmithImportContext::Init(TSharedRef< IDatasmithScene > InScene, const
 	{
 		SetupBaseOptionsVisibility();
 
-		bool bShouldImport = FDatasmithImportOptionHelper::DisplayOptionsDialog(ImportOptions, InScene);
+		bool bShouldImport = FDatasmithImportOptionHelper::DisplayOptionsDialog(ImportOptions, Scene.ToSharedRef());
 
 		ResetBaseOptionsVisibility();
 
-		if ( !bShouldImport )
+		if (!bShouldImport)
 		{
 			FDatasmithImportOptionHelper::CleanUpOptions(ImportOptions);
 			UE_LOG(LogDatasmithImport, Display, TEXT("Import canceled."));
@@ -288,16 +281,27 @@ bool FDatasmithImportContext::Init(TSharedRef< IDatasmithScene > InScene, const
 
 	if (SceneTranslator)
 	{
- 		AdditionalImportOptions.Push(Options);
+		AdditionalImportOptions.Push(Options);
 		SceneTranslator->SetSceneImportOptions(AdditionalImportOptions);
- 		AdditionalImportOptions.Pop();
+		AdditionalImportOptions.Pop();
 	}
 
-	Options->UpdateNotDisplayedConfig( bIsAReimport );
+	Options->UpdateNotDisplayedConfig(bIsAReimport);
 
 	FDatasmithImportOptionHelper::CleanUpOptions(ImportOptions);
+	return true;
+}
 
-	if ( !ActorsContext.ImportWorld )
+bool FDatasmithImportContext::SetupDestination(const FString& InImportPath, EObjectFlags InFlags, FFeedbackContext* InWarn, bool bSilent)
+{
+	if (!Scene.IsValid())
+	{
+		return false;
+	}
+
+	Options->BaseOptions.AssetOptions.PackagePath = FName(*InImportPath);
+
+	if (!ActorsContext.ImportWorld)
 	{
 		// User is asking to import model in a new level
 		// Check to see if there is nothing to save and act according to user's selection
@@ -337,7 +341,6 @@ bool FDatasmithImportContext::Init(TSharedRef< IDatasmithScene > InScene, const
 	}
 
 	FeedbackContext = InWarn;
-	Scene = InScene;
 
 	// Initialize the filtered scene as a copy of the original scene. We will use it to then filter out items to import.
 	FilteredScene = FDatasmithSceneFactory::DuplicateScene(Scene.ToSharedRef());
@@ -349,7 +352,7 @@ bool FDatasmithImportContext::Init(TSharedRef< IDatasmithScene > InScene, const
 
 	bool bResult = AssetsContext.Init();
 
-	if ( ShouldImportActors() )
+	if (ShouldImportActors())
 	{
 		bResult = bResult && ActorsContext.Init();
 	}
diff --git a/Engine/Plugins/Enterprise/DatasmithImporter/Source/DatasmithImporter/Public/DatasmithImportContext.h b/Engine/Plugins/Enterprise/DatasmithImporter/Source/DatasmithImporter/Public/DatasmithImportContext.h
--- a/Engine/Plugins/Enterprise/DatasmithImporter/Source/DatasmithImporter/Public/DatasmithImportContext.h
+++ b/Engine/Plugins/Enterprise/DatasmithImporter/Source/DatasmithImporter/Public/DatasmithImportContext.h
@@ -295,8 +295,31 @@ public:
 	 * @param InParent	The package in which we are
 	 * @param bSilent	Doesn't display the options dialog and skips other user input requests
 	 */
-	bool Init(TSharedRef< IDatasmithScene > Scene, const FString& InImportPath, EObjectFlags InFlags, FFeedbackContext* InWarn, const TSharedPtr<FJsonObject>& ImportSettingsJson, bool bSilent);
-	bool Init(const FString& FileName, TSharedRef< IDatasmithScene > Scene, const FString& InImportPath, EObjectFlags InFlags, FFeedbackContext* InWarn, const TSharedPtr<FJsonObject>& ImportSettingsJson, bool bSilent);
+	bool Init(TSharedRef< IDatasmithScene > InScene, const FString& InImportPath, EObjectFlags InFlags, FFeedbackContext* InWarn, const TSharedPtr<FJsonObject>& ImportSettingsJson, bool bSilent);
+
+	/**
+	 * First part of the Init process, replaces Init, and requires a call to SetupDestination after that.
+	 * Displays the options to the end-user (blocking call), and updates translator accordingly.
+	 * When silent, Json options are parsed instead.
+	 *
+	 * @param InScene              Scene that the context will use for the translation and import
+	 * @param ImportSettingsJson   When bSilent, options as json
+	 * @param bSilent              Flag that prevents options to be displayed and edited
+	 * @return false if the user canceled -> import should not occurs in that case.
+	 */
+	bool InitOptions(TSharedRef< IDatasmithScene > InScene, const TSharedPtr<FJsonObject>& ImportSettingsJson, bool bSilent);
+
+	/**
+	 * Second part of the Init process, replaces Init, and requires a call to InitOptions before that.
+	 * Setup destination packages
+	 *
+	 * @param InImportPath   Destination package's name
+	 * @param InFlags        Flags applyed to all generated objects durring the following import
+	 * @param InWarn         Feedback context for the following import
+	 * @param bSilent        When false, prompt the user to save dirty packages
+	 * @return false if the user canceled -> import should not occurs in that case.
+	 */
+	bool SetupDestination(const FString& InImportPath, EObjectFlags InFlags, FFeedbackContext* InWarn, bool bSilent);
 
 	/**
 	 * Replace or add options based on it's UClass.
@@ -322,8 +345,6 @@ public:
 	void DisplayMessages();
 
 private:
-	void SetFileName(const FString& FileName);
-
 	void SetupBaseOptionsVisibility();
 	void ResetBaseOptionsVisibility();
 
