Datasmith Python API: Added a TranslateScene method to populate the datasmith scene

diff --git a/Engine/Plugins/Enterprise/DatasmithImporter/Source/DatasmithImporter/Private/DatasmithBlueprintLibrary.cpp b/Engine/Plugins/Enterprise/DatasmithImporter/Source/DatasmithImporter/Private/DatasmithBlueprintLibrary.cpp
--- a/Engine/Plugins/Enterprise/DatasmithImporter/Source/DatasmithImporter/Private/DatasmithBlueprintLibrary.cpp
+++ b/Engine/Plugins/Enterprise/DatasmithImporter/Source/DatasmithImporter/Private/DatasmithBlueprintLibrary.cpp
@@ -149,7 +149,7 @@ UDatasmithSceneElement* UDatasmithSceneElement::ConstructDatasmithSceneFromFile(
 	TSharedRef< IDatasmithScene > Scene = FDatasmithSceneFactory::CreateScene(*Source.GetSceneName());
 	DatasmithSceneElement->SetDatasmithSceneElement(Scene);
 
-	bool bLoadConfig = false; //!IsAutomatedImport();
+	const bool bLoadConfig = false; // don't load values from ini files
 	DatasmithSceneElement->ImportContextPtr.Reset(new FDatasmithImportContext(Source.GetSourceFile(), bLoadConfig, GetLoggerName(), GetDisplayName(), TranslatableSource.GetTranslator()));
 
 	return DatasmithSceneElement;
@@ -187,7 +187,7 @@ UDatasmithSceneElement* UDatasmithSceneElement::GetExistingDatasmithScene(const
 		DatasmithSceneElement->SetDatasmithSceneElement(Scene);
 
 		// Setup pipe for reimport
-		bool bLoadConfig = false;
+		const bool bLoadConfig = false;
 		DatasmithSceneElement->ImportContextPtr.Reset(new FDatasmithImportContext(Source.GetSourceFile(), bLoadConfig, GetLoggerName(), GetDisplayName(), TranslatableSource.GetTranslator()));
 		FDatasmithImportContext& ImportContext = *DatasmithSceneElement->ImportContextPtr;
 		ImportContext.SceneAsset = SceneAsset;
@@ -196,18 +196,41 @@ UDatasmithSceneElement* UDatasmithSceneElement::GetExistingDatasmithScene(const
 
 		FString ImportPath = ImportContext.Options->BaseOptions.AssetOptions.PackagePath.ToString();
 
-		bool bIsSilent = false;
 		return DatasmithSceneElement;
 	}
 
 	return nullptr;
 }
 
+bool UDatasmithSceneElement::TranslateScene()
+{
+	if (!SourcePtr.IsValid() || !ImportContextPtr.IsValid() || !GetSceneElement().IsValid())
+	{
+		UE_LOG(LogDatasmithImport, Error, TEXT("Invalid State. Ensure ConstructDatasmithSceneFromFile has been called."));
+		return false;
+	}
+
+	FDatasmithTranslatableSceneSource& TranslatableSource = *SourcePtr;
+
+	TSharedRef<IDatasmithScene> Scene = GetSceneElement().ToSharedRef();
+	const bool bSilent = true; // don't pop options window
+	ImportContextPtr->InitOptions(Scene, nullptr, bSilent);
+
+	bTranslated = true;
+	if (!TranslatableSource.Translate(Scene))
+	{
+		UE_LOG(LogDatasmithImport, Error, TEXT("Datasmith import error: Scene translation failure. Abort import."));
+		return false;
+	}
+
+	return true;
+}
+
 FDatasmithImportFactoryCreateFileResult UDatasmithSceneElement::ImportScene(const FString& DestinationFolder)
 {
 	FDatasmithImportFactoryCreateFileResult Result;
 
-	if (this == nullptr || !ImportContextPtr.IsValid() || !SourcePtr.IsValid() || SourcePtr->GetTranslator() == nullptr || !GetSceneElement().IsValid())
+	if (this == nullptr || !ImportContextPtr.IsValid() || !GetSceneElement().IsValid())
 	{
 		UE_LOG(LogDatasmithImport, Error, TEXT("Invalid State. Ensure ConstructDatasmithSceneFromFile has been called."));
 		return Result;
@@ -221,23 +244,21 @@ FDatasmithImportFactoryCreateFileResult UDatasmithSceneElement::ImportScene(cons
 		return Result;
 	}
 
-	FDatasmithImportContext& ImportContext = *ImportContextPtr;
-	TSharedRef< IDatasmithScene > Scene = GetSceneElement().ToSharedRef();
-	EObjectFlags NewObjectFlags = RF_Public | RF_Standalone | RF_Transactional;
-	TSharedPtr<FJsonObject> ImportSettingsJson;
-	bool bIsSilent = true;
-	if ( !ImportContext.Init( Scene, DestinationPackage->GetName(), NewObjectFlags, GWarn, ImportSettingsJson, bIsSilent ) )
+	if (!bTranslated)
 	{
-		return Result;
+		if (!TranslateScene())
+		{
+			return Result;
+		}
 	}
 
-	FDatasmithTranslatableSceneSource& TranslatableSource = *SourcePtr;
-	if (!TranslatableSource.Translate(Scene))
+	FDatasmithImportContext& ImportContext = *ImportContextPtr;
+	const EObjectFlags NewObjectFlags = RF_Public | RF_Standalone | RF_Transactional;
+	const bool bIsSilent = true;
+	if (!ImportContext.SetupDestination(DestinationPackage->GetName(), NewObjectFlags, GWarn, bIsSilent))
 	{
-		UE_LOG(LogDatasmithImport, Error, TEXT("Datasmith import error: Scene translation failure. Abort import."));
 		return Result;
 	}
-
 	bool bUserCancelled = false;
 	Result.bImportSucceed = DatasmithImportFactoryImpl::ImportDatasmithScene(ImportContext, bUserCancelled);
 	Result.bImportSucceed &= !bUserCancelled;
@@ -276,7 +297,7 @@ FDatasmithImportFactoryCreateFileResult UDatasmithSceneElement::ReimportScene()
 	TSharedRef< IDatasmithScene > Scene = GetSceneElement().ToSharedRef();
 	EObjectFlags NewObjectFlags = RF_Public | RF_Standalone | RF_Transactional;
 	TSharedPtr<FJsonObject> ImportSettingsJson;
-	bool bIsSilent = true;
+	const bool bIsSilent = true;
 	if ( !ImportContext.Init( Scene, DestinationPackage->GetName(), NewObjectFlags, GWarn, ImportSettingsJson, bIsSilent ) )
 	{
 		return Result;
diff --git a/Engine/Plugins/Enterprise/DatasmithImporter/Source/DatasmithImporter/Public/DatasmithBlueprintLibrary.h b/Engine/Plugins/Enterprise/DatasmithImporter/Source/DatasmithImporter/Public/DatasmithBlueprintLibrary.h
--- a/Engine/Plugins/Enterprise/DatasmithImporter/Source/DatasmithImporter/Public/DatasmithBlueprintLibrary.h
+++ b/Engine/Plugins/Enterprise/DatasmithImporter/Source/DatasmithImporter/Public/DatasmithBlueprintLibrary.h
@@ -65,6 +65,13 @@ public:
 	UFUNCTION(BlueprintCallable, Category = "Editor Scripting | Datasmith")
 	static UDatasmithSceneElement* GetExistingDatasmithScene(const FString& AssetPath);
 
+	/**
+	 * Trigger the translation phase, which populates the DatasmithScene.
+	 * note that options should have been set before that.
+	 */
+	UFUNCTION(BlueprintCallable, Category = "Editor Scripting | Datasmith")
+	bool TranslateScene();
+
 	/**
 	 * Import a Datasmith Scene created with ConstructDatasmithSceneFromFile.
 	 * @param	DestinationFolder	Destination of where you want the asset to be imported. ie: /Game/MyFolder1
@@ -89,7 +96,7 @@ public:
 	UObject* GetOptions(UClass* OptionType=nullptr);
 
 	/**
-	 * List all the options stuctures applicable on this scene
+	 * List all the options structures applicable on this scene
 	 * @return A map of options type to option object
 	 */
 	UFUNCTION(BlueprintCallable, Category = "Datasmith | Scene")
@@ -109,6 +116,7 @@ public:
 private:
 	TUniquePtr<FDatasmithTranslatableSceneSource> SourcePtr; // #ueent_todo move to context
 	TUniquePtr<FDatasmithImportContext> ImportContextPtr;
+	bool bTranslated = false;
 };
 
 
