Fixed crash when output level already exists but is not loaded (finally)

#rnx

diff --git a/Engine/Plugins/Enterprise/DatasmithImporter/Source/DatasmithImporter/Private/DatasmithConsumer.cpp b/Engine/Plugins/Enterprise/DatasmithImporter/Source/DatasmithImporter/Private/DatasmithConsumer.cpp
--- a/Engine/Plugins/Enterprise/DatasmithImporter/Source/DatasmithImporter/Private/DatasmithConsumer.cpp
+++ b/Engine/Plugins/Enterprise/DatasmithImporter/Source/DatasmithImporter/Private/DatasmithConsumer.cpp
@@ -773,7 +773,21 @@ bool UDatasmithConsumer::CanCreateLevel(const FString& RequestedFolder, const FS
 
 	const FString AssetPathName = ObjectPath.GetAssetPathString();
 
+	bool bCanCreateAsset = false;
+
+	// Check in the asset registry to see if there could be any conflict
 	if(FDatasmithImporterUtils::CanCreateAsset(AssetPathName, UWorld::StaticClass()) == FDatasmithImporterUtils::EAssetCreationStatus::CS_CanCreate)
+	{
+		bCanCreateAsset = true;
+
+		// No conflict in the AssetRegistry, check if no unregistered asset in memory could conflict
+		if(UObject* Asset = ObjectPath.ResolveObject())
+		{
+			bCanCreateAsset = Cast<UWorld>(Asset) != nullptr;
+		}
+	}
+
+	if(bCanCreateAsset)
 	{
 		FText OutReason;
 		if(FDatasmithImporterImpl::CheckAssetPersistenceValidity(ObjectPath.GetLongPackageName(), *ImportContextPtr, FPackageName::GetMapPackageExtension(), OutReason))
@@ -946,9 +960,11 @@ ULevel* UDatasmithConsumer::FindOrAddLevel(const FString& InLevelName)
 	// The level is not part of the world
 	FString PackageFilename;
 	FPackageName::TryConvertLongPackageNameToFilename( LevelPackageName, PackageFilename, FPackageName::GetMapPackageExtension() );
+	FSoftObjectPath LevelSoftObjectPath(LevelPackageName + TEXT(".") + InLevelName);
 
 	ULevelStreaming* StreamingLevel = nullptr;
-	if(FPaths::FileExists(PackageFilename))
+	// If it already exists on disk or in memory, add it to the working world
+	if(FPaths::FileExists(PackageFilename) || LevelSoftObjectPath.ResolveObject())
 	{
 		FTransform LevelTransform;
 		StreamingLevel = UEditorLevelUtils::AddLevelToWorld(WorkingWorld.Get(), *LevelPackageName, ULevelStreamingAlwaysLoaded::StaticClass(), LevelTransform);
