DataPrep - Create Proxy Mesh discrepancy with merge actor

diff --git a/Engine/Plugins/Enterprise/DataprepEditor/Source/DataprepCore/Private/DataprepCoreUtils.cpp b/Engine/Plugins/Enterprise/DataprepEditor/Source/DataprepCore/Private/DataprepCoreUtils.cpp
--- a/Engine/Plugins/Enterprise/DataprepEditor/Source/DataprepCore/Private/DataprepCoreUtils.cpp
+++ b/Engine/Plugins/Enterprise/DataprepEditor/Source/DataprepCore/Private/DataprepCoreUtils.cpp
@@ -650,7 +650,23 @@ void FDataprepCoreUtils::BuildAssets(const TArray<TWeakObjectPtr<UObject>>& Asse
 			}
 			else if(UStaticMesh* StaticMesh = Cast<UStaticMesh>(AssetObject))
 			{
-				StaticMeshes.Add( StaticMesh );
+				// This ensures we have mesh description and thus avoid an assert when building meshes (StaticMesh.cpp:2470).
+				// We can clear the mesh descriptions post load.
+				const int32 NumLODs = StaticMesh->GetNumSourceModels();
+				bool bValidMesh = true;
+				for (int32 LODIndex = 0; LODIndex < NumLODs; ++LODIndex)
+				{
+					if ( !StaticMesh->GetMeshDescription(LODIndex) )
+					{
+						bValidMesh = false;
+						break;
+					}
+				}
+
+				if( bValidMesh )
+				{
+					StaticMeshes.Add( StaticMesh );
+				}
 			}
 		}
 	}
diff --git a/Engine/Plugins/Enterprise/DatasmithImporter/Source/DatasmithImporter/Private/DatasmithFileProducer.cpp b/Engine/Plugins/Enterprise/DatasmithImporter/Source/DatasmithImporter/Private/DatasmithFileProducer.cpp
--- a/Engine/Plugins/Enterprise/DatasmithImporter/Source/DatasmithImporter/Private/DatasmithFileProducer.cpp
+++ b/Engine/Plugins/Enterprise/DatasmithImporter/Source/DatasmithImporter/Private/DatasmithFileProducer.cpp
@@ -14,6 +14,7 @@
 #include "DatasmithScene.h"
 #include "DatasmithSceneActor.h"
 #include "DatasmithSceneFactory.h"
+#include "DatasmithStaticMeshImporter.h"
 #include "DatasmithTranslatorManager.h"
 #include "IDatasmithSceneElements.h"
 #include "Utility/DatasmithImporterUtils.h"
@@ -240,6 +241,8 @@ void UDatasmithFileProducer::SceneElementToWorld()
 	// STATIC MESHES
 	FDatasmithImporter::ImportStaticMeshes( *ImportContextPtr );
 
+	FDatasmithStaticMeshImporter::PreBuildStaticMeshes( *ImportContextPtr );
+
 	// ACTORS
 	{
 		FDatasmithImporter::ImportActors( *ImportContextPtr );
@@ -277,27 +280,6 @@ void UDatasmithFileProducer::SceneElementToWorld()
 	TArray<UStaticMesh*> StaticMeshes;
 	ImportContextPtr->ImportedStaticMeshes.GenerateValueArray(StaticMeshes);
 
-	{
-		TRACE_CPUPROFILER_EVENT_SCOPE(CommitMeshDescriptions);
-
-		UStaticMesh::FCommitMeshDescriptionParams Params;
-		Params.bMarkPackageDirty = false;
-		Params.bUseHashAsGuid = true;
-
-		ParallelFor(StaticMeshes.Num(),
-			[&](int32 StaticMeshIndex)
-			{
-				if(UStaticMesh* StaticMesh = StaticMeshes[StaticMeshIndex])
-				{
-					for (int32 Index = 0; Index < StaticMesh->GetNumSourceModels(); ++Index)
-					{
-						StaticMesh->CommitMeshDescription( Index, Params );
-					}
-				}
-			}
-		);
-	}
-
 	// Note: Some of the assets might be null (incomplete or failed import), only add non-null ones to Assets
 
 	for(UStaticMesh* StaticMesh : StaticMeshes)
