Fixed crash due to the existence of  unloaded assets with the same target name

#rnx

diff --git a/Engine/Plugins/Enterprise/DatasmithImporter/Source/DatasmithImporter/Private/DatasmithConsumer.cpp b/Engine/Plugins/Enterprise/DatasmithImporter/Source/DatasmithImporter/Private/DatasmithConsumer.cpp
--- a/Engine/Plugins/Enterprise/DatasmithImporter/Source/DatasmithImporter/Private/DatasmithConsumer.cpp
+++ b/Engine/Plugins/Enterprise/DatasmithImporter/Source/DatasmithImporter/Private/DatasmithConsumer.cpp
@@ -66,6 +66,7 @@
 #include "Misc/Paths.h"
 #include "Misc/SecureHash.h"
 #include "ObjectTools.h"
+#include "PackageTools.h"
 #include "UObject/Package.h"
 #include "UObject/SoftObjectPtr.h"
 #include "UObject/UnrealType.h"
@@ -317,58 +318,17 @@ bool UDatasmithConsumer::Initialize()
 
 	UpdateScene();
 
-	UPackage* ParentPackage = CreatePackage( nullptr, *GetTargetPackagePath() );
-	ParentPackage->FullyLoad();
-
 	// Re-create the DatasmithScene if it is invalid
 	if ( !DatasmithScene.IsValid() )
 	{
+		FString PackageName = FPaths::Combine( GetTargetPackagePath(), GetOuter()->GetName() + DatasmithSceneSuffix );
 
-		FString DatasmithSceneName = GetOuter()->GetName() + DatasmithSceneSuffix;
-
-		UPackage* Package = CreatePackage( nullptr, *FPaths::Combine( ParentPackage->GetPathName(), DatasmithSceneName ) );
-		Package->FullyLoad();
-
-		if( UObject* ExistingObject = StaticFindObject( nullptr, Package, *DatasmithSceneName, true ) )
-		{
-			bool bDatasmithSceneFound = false;
-
-			// Check to see if existing scene is not from same Dataprep asset
-			if( UDatasmithScene* ExistingDatasmithScene = Cast<UDatasmithScene>( ExistingObject ) )
-			{
-				if ( ExistingDatasmithScene->GetClass()->ImplementsInterface( UInterface_AssetUserData::StaticClass() ) )
-				{
-					if ( IInterface_AssetUserData* AssetUserDataInterface = Cast< IInterface_AssetUserData >( ExistingDatasmithScene ) )
-					{
-						if( UDataprepAssetUserData* DataprepAssetUserData = AssetUserDataInterface->GetAssetUserData< UDataprepAssetUserData >() )
-						{
-							UDataprepAssetInterface* DataprepAssetInterface = Cast< UDataprepAssetInterface >( GetOuter() );
-							check( DataprepAssetInterface );
+		UPackage* Package = UPackageTools::FindOrCreatePackageForAssetType( FName( *PackageName ), UDatasmithScene::StaticClass() );
+		check( Package );
 
-							if( DataprepAssetUserData->DataprepAssetPtr == DataprepAssetInterface )
-							{
-								DatasmithScene.Reset();
-								DatasmithScene = ExistingDatasmithScene;
-								Package = nullptr;
-								bDatasmithSceneFound = true;
-							}
-						}
-					}
-				}
-			}
-
-			if( !bDatasmithSceneFound )
-			{
-				DatasmithSceneName = MakeUniqueObjectName( ParentPackage, UDatasmithScene::StaticClass(), *DatasmithSceneName ).ToString();
-				Package = CreatePackage( nullptr, *FPaths::Combine( ParentPackage->GetPathName(), DatasmithSceneName ) );
-				Package->FullyLoad();
-			}
-		}
-
-		if(Package != nullptr)
-		{
-			DatasmithScene = NewObject< UDatasmithScene >( Package, *DatasmithSceneName, GetFlags() | RF_Standalone | RF_Public | RF_Transactional );
-		}
+		FString DatasmithSceneName = FPaths::GetBaseFilename( Package->GetFullName(), true );
+		
+		DatasmithScene = NewObject< UDatasmithScene >( Package, *DatasmithSceneName, GetFlags() | RF_Standalone | RF_Public | RF_Transactional );
 		check( DatasmithScene.IsValid() );
 
 		DatasmithScene->MarkPackageDirty();
