Fixed ensure when launching EngineTest in -game. Datasmith was trying to register it's UI menus when it shouldn't.

diff --git a/Engine/Plugins/Enterprise/DatasmithImporter/Source/DatasmithImporter/Private/DatasmithImporterModule.cpp b/Engine/Plugins/Enterprise/DatasmithImporter/Source/DatasmithImporter/Private/DatasmithImporterModule.cpp
--- a/Engine/Plugins/Enterprise/DatasmithImporter/Source/DatasmithImporter/Private/DatasmithImporterModule.cpp
+++ b/Engine/Plugins/Enterprise/DatasmithImporter/Source/DatasmithImporter/Private/DatasmithImporterModule.cpp
@@ -63,7 +63,7 @@ public:
 		UDatasmithFileProducer::LoadDefaultSettings();
 
 		// Disable any UI feature if running in command mode
-		if (!IsRunningCommandlet())
+		if (UToolMenus::IsToolMenuUIEnabled())
 		{
 			FDatasmithUIManager::Initialize();
 
@@ -90,7 +90,7 @@ public:
 	virtual void ShutdownModule() override
 	{
 		// Disable any UI feature if running in command mode
-		if (!IsRunningCommandlet())
+		if (UToolMenus::IsToolMenuUIEnabled())
 		{
 			RemoveDataprepMenuEntryForDatasmithSceneAsset();
 
