Updated Datasmith Texture Importer to use the Interchange Framework
- Datasmith will only wait on the texture import result when it needs it for another asset (ie: material).
- Added a custom texture pipeline to set the datasmith attributes on the texture node.
- Created a FAsyncResult that behaves like a TFuture but calling Get will actually tick the task graph so that we don't deadlock when trying to get the future from the game thread while it depends on a game thread task.

diff --git a/Engine/Plugins/Enterprise/DatasmithImporter/DatasmithImporter.uplugin b/Engine/Plugins/Enterprise/DatasmithImporter/DatasmithImporter.uplugin
--- a/Engine/Plugins/Enterprise/DatasmithImporter/DatasmithImporter.uplugin
+++ b/Engine/Plugins/Enterprise/DatasmithImporter/DatasmithImporter.uplugin
@@ -44,6 +44,10 @@
 			"Name": "DataPrepEditor",
 			"Enabled": true
 		},
+		{
+			"Name": "Interchange",
+			"Enabled": true
+		},
 		{
 			"Name": "VariantManagerContent",
 			"Enabled": true
diff --git a/Engine/Plugins/Enterprise/DatasmithImporter/Source/DatasmithImporter/DatasmithImporter.Build.cs b/Engine/Plugins/Enterprise/DatasmithImporter/Source/DatasmithImporter/DatasmithImporter.Build.cs
--- a/Engine/Plugins/Enterprise/DatasmithImporter/Source/DatasmithImporter/DatasmithImporter.Build.cs
+++ b/Engine/Plugins/Enterprise/DatasmithImporter/Source/DatasmithImporter/DatasmithImporter.Build.cs
@@ -24,6 +24,8 @@ namespace UnrealBuildTool.Rules
 					"Foliage",
 					"FreeImage",
 					"InputCore",
+					"InterchangeCore",
+					"InterchangeNodePlugin",
 					"Json",
 					"Landscape",
 					"LandscapeEditor",
@@ -49,7 +51,7 @@ namespace UnrealBuildTool.Rules
 					"UnrealEd",
 					"VariantManager",
 					"VariantManagerContent",
-                }
+				}
 			);
 
 			PublicDependencyModuleNames.AddRange(
@@ -57,11 +59,11 @@ namespace UnrealBuildTool.Rules
 				{
 					"DataprepCore",
 					"DatasmithContent",
-                    "DatasmithCore",
+					"DatasmithCore",
 					"DatasmithTranslator",
-                    "DatasmithContentEditor",
+					"DatasmithContentEditor",
 				}
 			);
-        }
-    }
+		}
+	}
 }
\ No newline at end of file
diff --git a/Engine/Plugins/Enterprise/DatasmithImporter/Source/DatasmithImporter/Private/DatasmithConsumer.cpp b/Engine/Plugins/Enterprise/DatasmithImporter/Source/DatasmithImporter/Private/DatasmithConsumer.cpp
--- a/Engine/Plugins/Enterprise/DatasmithImporter/Source/DatasmithImporter/Private/DatasmithConsumer.cpp
+++ b/Engine/Plugins/Enterprise/DatasmithImporter/Source/DatasmithImporter/Private/DatasmithConsumer.cpp
@@ -689,7 +689,7 @@ bool UDatasmithConsumer::BuildContexts()
 	ImportContextPtr->ActorsContext.FinalWorld = WorkingWorld.Get();
 
 	// Initialize ActorsContext's UniqueNameProvider with actors in the GWorld not the Import world
-	ImportContextPtr->ActorsContext.UniqueNameProvider = FDatasmithActorUniqueLabelProvider();
+	ImportContextPtr->ActorsContext.UniqueNameProvider.Clear();
 	ImportContextPtr->ActorsContext.UniqueNameProvider.PopulateLabelFrom( ImportContextPtr->ActorsContext.FinalWorld );
 
 	// Add assets as if they have been imported using the current import context
@@ -1528,7 +1528,8 @@ namespace DatasmithConsumerUtils
 					TSharedRef< IDatasmithTextureElement > TextureElement = FDatasmithSceneFactory::CreateTexture( *AssetTag );
 					TextureElement->SetLabel( *Texture->GetName() );
 
-					ImportContext.ImportedTextures.Add( TextureElement, Texture );
+					TPromise< UObject* > TexturePromise = MakeFulfilledPromise< UObject* >( Texture );
+					ImportContext.ImportedTextures.FindOrAdd( TextureElement ) = Interchange::FAsyncImportResult{ TexturePromise.GetFuture(), FGraphEventRef() };
 					ImportContext.Scene->AddTexture( TextureElement );
 				}
 				else if(UMaterialInstance* MaterialInstance = Cast<UMaterialInstance>(Asset))
diff --git a/Engine/Plugins/Enterprise/DatasmithImporter/Source/DatasmithImporter/Private/DatasmithFileProducer.cpp b/Engine/Plugins/Enterprise/DatasmithImporter/Source/DatasmithImporter/Private/DatasmithFileProducer.cpp
--- a/Engine/Plugins/Enterprise/DatasmithImporter/Source/DatasmithImporter/Private/DatasmithFileProducer.cpp
+++ b/Engine/Plugins/Enterprise/DatasmithImporter/Source/DatasmithImporter/Private/DatasmithFileProducer.cpp
@@ -293,11 +293,11 @@ void UDatasmithFileProducer::SceneElementToWorld()
 		}
 	}
 
-	for ( TPair< TSharedRef< IDatasmithTextureElement >, UTexture* >& AssetPair : ImportContextPtr->ImportedTextures )
+	for ( TPair< TSharedRef< IDatasmithTextureElement >, Interchange::FAsyncImportResult >& AssetPair : ImportContextPtr->ImportedTextures )
 	{
-		if(AssetPair.Value)
+		if(AssetPair.Value.IsValid())
 		{
-			Assets.Emplace( AssetPair.Value );
+			Assets.Emplace( AssetPair.Value.Get() );
 		}
 	}
 
diff --git a/Engine/Plugins/Enterprise/DatasmithImporter/Source/DatasmithImporter/Private/DatasmithImporter.cpp b/Engine/Plugins/Enterprise/DatasmithImporter/Source/DatasmithImporter/Private/DatasmithImporter.cpp
--- a/Engine/Plugins/Enterprise/DatasmithImporter/Source/DatasmithImporter/Private/DatasmithImporter.cpp
+++ b/Engine/Plugins/Enterprise/DatasmithImporter/Source/DatasmithImporter/Private/DatasmithImporter.cpp
@@ -331,6 +331,34 @@ void FDatasmithImporter::ImportTextures( FDatasmithImportContext& ImportContext
 
 		FDatasmithTextureResize::Initialize();
 
+		// Try to import textures async first, if possible 
+		for ( int32 TextureIndex = 0; TextureIndex < FilteredTextureElements.Num(); TextureIndex++ )
+		{
+			ImportContext.bUserCancelled |= FDatasmithImporterImpl::HasUserCancelledTask( ImportContext.FeedbackContext );
+
+			if (ImportContext.bUserCancelled)
+			{
+				break;
+			}
+
+			Interchange::FAsyncImportResult FutureTexture = DatasmithTextureImporter.CreateTextureAsync( FilteredTextureElements[TextureIndex] );
+
+			if ( FutureTexture.IsValid() )
+			{
+				FutureTexture = FutureTexture.Next(
+					[ TextureElement = FilteredTextureElements[TextureIndex].ToSharedRef(), &ImportContext ]( UObject* Future )
+					{
+						ImportMetaDataForObject( ImportContext, TextureElement, Future );
+
+						return Future;
+					}
+				);
+
+				Interchange::FAsyncImportResult& ImportedTexture = ImportContext.ImportedTextures.FindOrAdd( FilteredTextureElements[TextureIndex].ToSharedRef() );
+				ImportedTexture = MoveTemp( FutureTexture );
+			}
+		}
+
 		struct FAsyncData
 		{
 			FString       Extension;
@@ -342,6 +370,11 @@ void FDatasmithImporter::ImportTextures( FDatasmithImportContext& ImportContext
 
 		for ( int32 TextureIndex = 0; TextureIndex < FilteredTextureElements.Num(); TextureIndex++ )
 		{
+			if ( ImportContext.ImportedTextures.Contains( FilteredTextureElements[TextureIndex].ToSharedRef() ) )
+			{
+				continue;
+			}
+
 			ImportContext.bUserCancelled |= FDatasmithImporterImpl::HasUserCancelledTask( ImportContext.FeedbackContext );
 
 			AsyncData[TextureIndex].Result = 
@@ -364,6 +397,11 @@ void FDatasmithImporter::ImportTextures( FDatasmithImportContext& ImportContext
 
 		for ( int32 TextureIndex = 0; TextureIndex < FilteredTextureElements.Num(); TextureIndex++ )
 		{
+			if ( ImportContext.ImportedTextures.Contains( FilteredTextureElements[TextureIndex].ToSharedRef() ) )
+			{
+				continue;
+			}
+
 			ImportContext.bUserCancelled |= FDatasmithImporterImpl::HasUserCancelledTask( ImportContext.FeedbackContext );
 
 			if ( ImportContext.bUserCancelled )
@@ -405,18 +443,19 @@ UTexture* FDatasmithImporter::ImportTexture( FDatasmithImportContext& ImportCont
 {
 	TRACE_CPUPROFILER_EVENT_SCOPE(FDatasmithImporter::ImportTexture);
 
-	UTexture*& ImportedTexture = ImportContext.ImportedTextures.FindOrAdd( TextureElement );
-	ImportedTexture = DatasmithTextureImporter.CreateTexture( TextureElement, TextureData, Extension );
+	Interchange::FAsyncImportResult& ImportedTexture = ImportContext.ImportedTextures.FindOrAdd( TextureElement );
+	TPromise< UObject* > TexturePromise = MakeFulfilledPromise< UObject* >( DatasmithTextureImporter.CreateTexture( TextureElement, TextureData, Extension ) );
+	ImportedTexture = Interchange::FAsyncImportResult{ TexturePromise.GetFuture(), FGraphEventRef() };
 
-	if (ImportedTexture == nullptr)
+	if ( !ImportedTexture.Get() )
 	{
 		ImportContext.ImportedTextures.Remove( TextureElement );
 		return nullptr;
 	}
 
-	ImportMetaDataForObject( ImportContext, TextureElement, ImportedTexture );
+	ImportMetaDataForObject( ImportContext, TextureElement, ImportedTexture.Get() );
 
-	return ImportedTexture;
+	return Cast< UTexture >( ImportedTexture.Get() );
 }
 
 UTexture* FDatasmithImporter::FinalizeTexture( UTexture* SourceTexture, const TCHAR* TexturesFolderPath, UTexture* ExistingTexture, TMap< UObject*, UObject* >* ReferencesToRemap )
@@ -1459,14 +1498,14 @@ void FDatasmithImporter::FinalizeImport(FDatasmithImportContext& ImportContext,
 	FScopedSlowTask* Progress = ProgressPtr.Get();
 
 	// Needs to be done in dependencies order (textures -> materials -> static meshes)
-	for (const TPair< TSharedRef< IDatasmithTextureElement >, UTexture* >& ImportedTexturePair : ImportContext.ImportedTextures)
+	for (const TPair< TSharedRef< IDatasmithTextureElement >, Interchange::FAsyncImportResult >& ImportedTexturePair : ImportContext.ImportedTextures)
 	{
 		if (ImportContext.bUserCancelled)
 		{
 			break;
 		}
 
-		UTexture* SourceTexture = ImportedTexturePair.Value;
+		UTexture* SourceTexture = Cast< UTexture >( ImportedTexturePair.Value.Get() );
 
 		if (!SourceTexture || (ValidAssets.Num() > 0 && !ValidAssets.Contains(Cast<UObject>(SourceTexture))))
 		{
diff --git a/Engine/Plugins/Enterprise/DatasmithImporter/Source/DatasmithImporter/Private/DatasmithTextureImporter.cpp b/Engine/Plugins/Enterprise/DatasmithImporter/Source/DatasmithImporter/Private/DatasmithTextureImporter.cpp
--- a/Engine/Plugins/Enterprise/DatasmithImporter/Source/DatasmithImporter/Private/DatasmithTextureImporter.cpp
+++ b/Engine/Plugins/Enterprise/DatasmithImporter/Source/DatasmithImporter/Private/DatasmithTextureImporter.cpp
@@ -30,8 +30,10 @@
 
 #include "Modules/ModuleManager.h"
 
+#include "InterchangeManager.h"
 #include "ObjectTools.h"
 #include "RHI.h"
+#include "TextureNode.h"
 
 #define LOCTEXT_NAMESPACE "DatasmithTextureImport"
 
@@ -306,4 +308,155 @@ UTexture* FDatasmithTextureImporter::CreateTexture(const TSharedPtr<IDatasmithTe
 	return Texture;
 }
 
+Interchange::FAsyncImportResult FDatasmithTextureImporter::CreateTextureAsync(const TSharedPtr<IDatasmithTextureElement>& TextureElement)
+{
+	Interchange::FScopedSourceData ScopedSourceData( TextureElement->GetFile() );
+
+	const FString ContentPath = ImportContext.AssetsContext.TexturesImportPackage->GetPathName();
+
+	FImportAssetParameters ImportAssetParameters;
+	ImportAssetParameters.bIsAutomated = true; // From the InterchangeManager point of view, this is considered an automated import
+
+	UDatasmithTexturePipeline* TexturePipeline = NewObject< UDatasmithTexturePipeline >();
+	TexturePipeline->TextureElement = TextureElement;
+
+	ImportAssetParameters.OverridePipeline = TexturePipeline;
+
+	Interchange::FAsyncImportResult FutureTexture = UInterchangeManager::GetInterchangeManager().ImportAssetAsync( ContentPath, ScopedSourceData.GetSourceData(), ImportAssetParameters );
+
+	return FutureTexture;
+}
+
+bool UDatasmithTexturePipeline::ExecuteImportPipeline(UInterchangeBaseNodeContainerAdapter* BaseNodeContainerAdapter)
+{
+	if ( !TextureElement.IsValid() )
+	{
+		return false;
+	}
+
+	TArray< FName > Nodes;
+	BaseNodeContainerAdapter->GetRoots( Nodes );
+
+	if ( Nodes.Num() <= 0)
+	{
+		return false;
+	}
+
+	Interchange::FTextureNode& TextureNode = static_cast< Interchange::FTextureNode& >( BaseNodeContainerAdapter->BaseNodeContainer->GetNodeChecked( Nodes[0] ) );
+
+	TOptional< bool > bFlipNormalMapGreenChannel;
+	TOptional< TextureMipGenSettings > MipGenSettings;
+	TOptional< TextureGroup > LODGroup;
+	TOptional< TextureCompressionSettings > CompressionSettings;
+
+	// Make sure to set the proper LODGroup as it's used to determine the CompressionSettings when using TEXTUREGROUP_WorldNormalMap
+	switch ( TextureElement->GetTextureMode() )
+	{
+	case EDatasmithTextureMode::Diffuse:
+		MipGenSettings = TMGS_Sharpen5;
+		LODGroup = TEXTUREGROUP_World;
+		break;
+	case EDatasmithTextureMode::Specular:
+		LODGroup = TEXTUREGROUP_WorldSpecular;
+		break;
+	case EDatasmithTextureMode::Bump:
+	case EDatasmithTextureMode::Normal:
+		LODGroup = TEXTUREGROUP_WorldNormalMap;
+		CompressionSettings = TC_Normalmap;
+		break;
+	case EDatasmithTextureMode::NormalGreenInv:
+		LODGroup = TEXTUREGROUP_WorldNormalMap;
+		CompressionSettings = TC_Normalmap;
+		bFlipNormalMapGreenChannel = true;
+		break;
+	case EDatasmithTextureMode::Displace:
+		LODGroup = TEXTUREGROUP_World;
+		CompressionSettings = TC_Displacementmap;
+		break;
+	}
+
+	const TOptional< float > RGBCurve = [ this ]() -> TOptional< float >
+	{
+		const float ElementRGBCurve = TextureElement->GetRGBCurve();
+
+		if ( FMath::IsNearlyEqual( ElementRGBCurve, 1.0f ) == false && ElementRGBCurve > 0.f )
+		{
+			return ElementRGBCurve;
+		}
+		else
+		{
+			return {};
+		}
+	}();
+
+	static_assert(TextureAddress::TA_Wrap == (int)EDatasmithTextureAddress::Wrap && TextureAddress::TA_Mirror == (int)EDatasmithTextureAddress::Mirror, "Texture Address enum doesn't match!" );
+
+	TOptional< TextureFilter > TexFilter;
+
+	switch ( TextureElement->GetTextureFilter() )
+	{
+	case EDatasmithTextureFilter::Nearest:
+		TexFilter = TextureFilter::TF_Nearest;
+		break;
+	case EDatasmithTextureFilter::Bilinear:
+		TexFilter = TextureFilter::TF_Bilinear;
+		break;
+	case EDatasmithTextureFilter::Trilinear:
+		TexFilter = TextureFilter::TF_Trilinear;
+		break;
+	}
+
+	TOptional< bool > bSrgb;
+
+	EDatasmithColorSpace ColorSpace = TextureElement->GetSRGB();
+	if ( ColorSpace == EDatasmithColorSpace::sRGB )
+	{
+		bSrgb = true;
+	}
+	else if ( ColorSpace == EDatasmithColorSpace::Linear )
+	{
+		bSrgb = false;
+	}
+
+	TextureNode.SetCustomAddressX( (TextureAddress)TextureElement->GetTextureAddressX() );
+	TextureNode.SetCustomAddressY( (TextureAddress)TextureElement->GetTextureAddressY() );
+
+	if ( bSrgb.IsSet() )
+	{
+		TextureNode.SetCustomSRGB( bSrgb.GetValue() );
+	}
+
+	if ( bFlipNormalMapGreenChannel.IsSet() )
+	{
+		TextureNode.SetCustombFlipGreenChannel( bFlipNormalMapGreenChannel.GetValue() );
+	}
+
+	if ( MipGenSettings.IsSet() )
+	{
+		TextureNode.SetCustomMipGenSettings( MipGenSettings.GetValue() );
+	}
+
+	if ( LODGroup.IsSet() )
+	{
+		TextureNode.SetCustomLODGroup( LODGroup.GetValue() );
+	}
+
+	if ( CompressionSettings.IsSet() )
+	{
+		TextureNode.SetCustomLODGroup( CompressionSettings.GetValue() );
+	}
+
+	if ( RGBCurve.IsSet() )
+	{
+		TextureNode.SetCustomAdjustRGBCurve( RGBCurve.GetValue() );
+	}
+
+	if ( TexFilter.IsSet() )
+	{
+		TextureNode.SetCustomFilter( TexFilter.GetValue() );
+	}
+
+	return true;
+}
+
 #undef LOCTEXT_NAMESPACE
diff --git a/Engine/Plugins/Enterprise/DatasmithImporter/Source/DatasmithImporter/Public/DatasmithImportContext.h b/Engine/Plugins/Enterprise/DatasmithImporter/Source/DatasmithImporter/Public/DatasmithImportContext.h
--- a/Engine/Plugins/Enterprise/DatasmithImporter/Source/DatasmithImporter/Public/DatasmithImportContext.h
+++ b/Engine/Plugins/Enterprise/DatasmithImporter/Source/DatasmithImporter/Public/DatasmithImportContext.h
@@ -13,6 +13,7 @@
 
 #include "DatasmithUtils.h"
 #include "DatasmithImportOptions.h"
+#include "InterchangeManager.h"
 
 class AActor;
 class ADatasmithSceneActor;
@@ -52,7 +53,7 @@ class IDatasmithTranslator;
 class DATASMITHIMPORTER_API FDatasmithActorUniqueLabelProvider : public FDatasmithUniqueNameProvider
 {
 public:
-	FDatasmithActorUniqueLabelProvider(UWorld* World=nullptr);
+	explicit FDatasmithActorUniqueLabelProvider(UWorld* World=nullptr);
 	void PopulateLabelFrom(UWorld* World);
 };
 
@@ -255,7 +256,7 @@ struct DATASMITHIMPORTER_API FDatasmithImportContext
 	TMap< FString, TSharedRef < IDatasmithMeshElement > > ImportedStaticMeshesByName;
 
 	/** Map of imported texture for each texture element */
-	TMap< TSharedRef< IDatasmithTextureElement >, UTexture* > ImportedTextures;
+	TMap< TSharedRef< IDatasmithTextureElement >, Interchange::FAsyncImportResult > ImportedTextures;
 
 	/** Map of imported material for each material element */
 	TMap< TSharedRef< IDatasmithBaseMaterialElement >, UMaterialInterface* > ImportedMaterials;
diff --git a/Engine/Plugins/Enterprise/DatasmithImporter/Source/DatasmithImporter/Public/DatasmithTextureImporter.h b/Engine/Plugins/Enterprise/DatasmithImporter/Source/DatasmithImporter/Public/DatasmithTextureImporter.h
--- a/Engine/Plugins/Enterprise/DatasmithImporter/Source/DatasmithImporter/Public/DatasmithTextureImporter.h
+++ b/Engine/Plugins/Enterprise/DatasmithImporter/Source/DatasmithImporter/Public/DatasmithTextureImporter.h
@@ -1,16 +1,27 @@
 // Copyright Epic Games, Inc. All Rights Reserved.
 #pragma once
 
+#include "Async/Future.h"
 #include "Containers/Array.h"
 #include "Templates/SharedPointer.h"
 #include "UObject/ObjectMacros.h"
 #include "UObject/StrongObjectPtr.h"
 
+#include "InterchangePipelineBase.h"
+#include "Async/TaskGraphInterfaces.h"
+
+#include "DatasmithTextureImporter.generated.h"
+
 struct FDatasmithImportContext;
 class IDatasmithTextureElement;
 class UTexture;
 class UTextureFactory;
 
+namespace Interchange
+{
+	class FAsyncImportResult;
+}
+
 class FDatasmithTextureImporter : private FNoncopyable
 {
 public:
@@ -20,6 +31,8 @@ public:
 	UTexture* CreateTexture(const TSharedPtr<IDatasmithTextureElement>& TextureElement, const TArray<uint8>& TextureData, const FString& Extension);
 	bool GetTextureData(const TSharedPtr<IDatasmithTextureElement>& TextureElement, TArray<uint8>& TextureData, FString& Extension);
 
+	Interchange::FAsyncImportResult CreateTextureAsync(const TSharedPtr<IDatasmithTextureElement>& TextureElement);
+
 private:
 	bool ResizeTextureElement(const TSharedPtr<IDatasmithTextureElement>& TextureElement, FString& ResizedFilename);
 
@@ -27,4 +40,15 @@ private:
 	FDatasmithImportContext& ImportContext;
 	TStrongObjectPtr< UTextureFactory > TextureFact;
 	FString TempDir;
+};
+
+UCLASS()
+class UDatasmithTexturePipeline : public UInterchangePipelineBase
+{
+	GENERATED_BODY()
+
+public:
+	virtual bool ExecuteImportPipeline(UInterchangeBaseNodeContainerAdapter* BaseNodeContainerAdapter) override;
+
+	TSharedPtr< IDatasmithTextureElement > TextureElement;
 };
\ No newline at end of file
diff --git a/Engine/Plugins/Enterprise/DatasmithImporter/Source/DatasmithImporter/Public/Utility/DatasmithImporterUtils.h b/Engine/Plugins/Enterprise/DatasmithImporter/Source/DatasmithImporter/Public/Utility/DatasmithImporterUtils.h
--- a/Engine/Plugins/Enterprise/DatasmithImporter/Source/DatasmithImporter/Public/Utility/DatasmithImporterUtils.h
+++ b/Engine/Plugins/Enterprise/DatasmithImporter/Source/DatasmithImporter/Public/Utility/DatasmithImporterUtils.h
@@ -4,6 +4,8 @@
 #include "CoreMinimal.h"
 #include "DatasmithImportContext.h"
 #include "DatasmithScene.h"
+#include "IDatasmithSceneElements.h"
+#include "InterchangeManager.h"
 #include "UObject/UObjectHash.h"
 
 class ADatasmithSceneActor;
@@ -248,7 +250,7 @@ struct FDatasmithFindAssetTypeHelper< UStaticMesh >
 template<>
 struct FDatasmithFindAssetTypeHelper< UTexture >
 {
-	static const TMap< TSharedRef< IDatasmithTextureElement >, UTexture* >& GetImportedAssetsMap( const FDatasmithAssetsImportContext& AssetsContext )
+	static const TMap< TSharedRef< IDatasmithTextureElement >, Interchange::FAsyncImportResult >& GetImportedAssetsMap( const FDatasmithAssetsImportContext& AssetsContext )
 	{
 		return AssetsContext.ParentContext.ImportedTextures;
 	}
@@ -408,6 +410,65 @@ inline ObjectType* FDatasmithImporterUtils::FindAsset( const FDatasmithAssetsImp
 	}
 }
 
+template<>
+inline UTexture* FDatasmithImporterUtils::FindAsset< UTexture >( const FDatasmithAssetsImportContext& AssetsContext, const TCHAR* ObjectPathName )
+{
+	if ( FCString::Strlen( ObjectPathName ) <= 0 )
+	{
+		return nullptr;
+	}
+
+	if ( FPaths::IsRelative( ObjectPathName ) )
+	{
+		const TSharedRef< IDatasmithTextureElement >* ImportedElement   = FDatasmithFindAssetTypeHelper< UTexture >::GetImportedElementByName( AssetsContext, ObjectPathName );
+		const TMap< TSharedRef< IDatasmithTextureElement >, Interchange::FAsyncImportResult >& ImportedAssetsMap = FDatasmithFindAssetTypeHelper< UTexture >::GetImportedAssetsMap( AssetsContext );
+
+		if (ImportedElement)
+		{
+			Interchange::FAsyncImportResult* Result =  const_cast< Interchange::FAsyncImportResult* >( ImportedAssetsMap.Find(*ImportedElement) );
+			ensure(Result);
+			if (Result)
+			{
+				return Cast< UTexture >( Result->Get() );
+			}
+		}
+		else
+		{
+			for ( const TPair< TSharedRef< IDatasmithTextureElement >, Interchange::FAsyncImportResult >& ImportedAssetPair : ImportedAssetsMap )
+			{
+				if ( FCString::Stricmp( ImportedAssetPair.Key->GetName(), ObjectPathName ) == 0 )
+				{
+					Interchange::FAsyncImportResult& Result = const_cast< Interchange::FAsyncImportResult& >( ImportedAssetPair.Value );
+					return Cast< UTexture >( Result.Get() );
+				}
+			}
+		}
+
+		{
+			const TMap< FName, TSoftObjectPtr< UTexture > >* AssetsMap = FDatasmithFindAssetTypeHelper< UTexture >::GetAssetsMap( AssetsContext.ParentContext.SceneAsset );
+
+			// Check if the AssetsMap is already tracking our asset
+			if ( AssetsMap && AssetsMap->Contains( ObjectPathName ) )
+			{
+				UTexture* Texture = (*AssetsMap)[ FName( ObjectPathName ) ].LoadSynchronous();
+				ensure( Texture );
+				return Texture;
+			}
+			else
+			{
+				UPackage* FinalPackage = FDatasmithFindAssetTypeHelper< UTexture >::GetFinalPackage( AssetsContext );
+				UTexture* Texture = FindObject< UTexture >( FinalPackage, ObjectPathName );
+				ensure( Texture );
+				return Texture;
+			}
+		}
+	}
+	else
+	{
+		return FindObject< UTexture >( nullptr, ObjectPathName );
+	}
+}
+
 template< class ObjectType >
 inline ObjectType* FDatasmithImporterUtils::FindObject( const UPackage* ParentPackage, const FString& ObjectName )
 {
diff --git a/Engine/Source/Runtime/Datasmith/DatasmithCore/Private/DatasmithUtils.cpp b/Engine/Source/Runtime/Datasmith/DatasmithCore/Private/DatasmithUtils.cpp
--- a/Engine/Source/Runtime/Datasmith/DatasmithCore/Private/DatasmithUtils.cpp
+++ b/Engine/Source/Runtime/Datasmith/DatasmithCore/Private/DatasmithUtils.cpp
@@ -650,6 +650,8 @@ bool FDatasmithSceneUtils::IsPostProcessUsedInScene(const TSharedPtr<IDatasmithS
 
 FString FDatasmithUniqueNameProviderBase::GenerateUniqueName(const FString& BaseName)
 {
+	FScopeLock Lock( &CriticalSection );
+
 	const int32 FrequentlyUsedThreshold = 5; // don't saturate the table with uncommon names
 	if (!Contains(BaseName))
 	{
diff --git a/Engine/Source/Runtime/Datasmith/DatasmithCore/Public/DatasmithUtils.h b/Engine/Source/Runtime/Datasmith/DatasmithCore/Public/DatasmithUtils.h
--- a/Engine/Source/Runtime/Datasmith/DatasmithCore/Public/DatasmithUtils.h
+++ b/Engine/Source/Runtime/Datasmith/DatasmithCore/Public/DatasmithUtils.h
@@ -238,6 +238,11 @@ public:
 	 */
 	virtual void RemoveExistingName(const FString& Name) = 0;
 
+	/**
+	 * Flushes all known names
+	 */
+	virtual void Clear() { FrequentlyUsedNames.Empty(); }
+
 protected:
 
 	/**
@@ -250,6 +255,7 @@ protected:
 
 private:
 	TMap<FString, int32> FrequentlyUsedNames;
+	FCriticalSection CriticalSection;
 };
 
 /**
@@ -258,14 +264,17 @@ private:
 class DATASMITHCORE_API FDatasmithUniqueNameProvider : public FDatasmithUniqueNameProviderBase
 {
 public:
+	using Super = FDatasmithUniqueNameProviderBase;
+
 	void Reserve( int32 NumberOfName ) { KnownNames.Reserve(NumberOfName); }
 
 	virtual void AddExistingName(const FString& Name) override { KnownNames.Add(Name); }
 	virtual void RemoveExistingName(const FString& Name) override { KnownNames.Remove(Name); }
 
+	virtual void Clear() override { Super::Clear(); KnownNames.Empty(); }
+
 protected:
 	virtual bool Contains(const FString& Name) override { return KnownNames.Contains(Name); }
-	void Clear() { KnownNames.Empty(); }
 
 private:
 	TSet<FString> KnownNames;
diff --git a/Engine/Source/Runtime/Interchange/Core/Private/InterchangeManager.cpp b/Engine/Source/Runtime/Interchange/Core/Private/InterchangeManager.cpp
--- a/Engine/Source/Runtime/Interchange/Core/Private/InterchangeManager.cpp
+++ b/Engine/Source/Runtime/Interchange/Core/Private/InterchangeManager.cpp
@@ -66,6 +66,11 @@ UInterchangeTranslatorBase* Interchange::FScopedTranslator::GetTranslator()
 	return ScopedTranslatorPtr.Get();
 }
 
+Interchange::FImportAsyncHelper::FImportAsyncHelper()
+{
+	RootObjectCompletionEvent = FGraphEvent::CreateGraphEvent();
+}
+
 void Interchange::FImportAsyncHelper::AddReferencedObjects(FReferenceCollector& Collector)
 {
 	for (UInterchangeSourceData* SourceData : SourceDatas)
@@ -129,6 +134,33 @@ void Interchange::FImportAsyncHelper::CleanUp()
 	Factories.Empty();
 }
 
+Interchange::FAsyncImportResult::FAsyncImportResult( TFuture< UObject* >&& InFutureObject, const FGraphEventRef& InGraphEvent )
+	: FutureObject( MoveTemp( InFutureObject ) )
+	, GraphEvent( InGraphEvent )
+{
+}
+
+bool Interchange::FAsyncImportResult::IsValid() const
+{
+	return FutureObject.IsValid();
+}
+
+UObject* Interchange::FAsyncImportResult::Get() const
+{
+	if ( !FutureObject.IsReady() )
+	{
+		// Tick the task graph until our FutureObject is ready
+		FTaskGraphInterface::Get().WaitUntilTaskCompletes( GraphEvent );
+	}
+
+	return FutureObject.Get();
+}
+
+Interchange::FAsyncImportResult Interchange::FAsyncImportResult::Next( TFunction< UObject*( UObject* ) > Continuation )
+{
+	return Interchange::FAsyncImportResult{ FutureObject.Next( Continuation ), GraphEvent };
+}
+
 void Interchange::SanitizeInvalidChar(FString& String)
 {
 	const TCHAR* InvalidChar = INVALID_OBJECTPATH_CHARACTERS;
@@ -213,6 +245,11 @@ bool UInterchangeManager::CanTranslateSourceData(const UInterchangeSourceData* S
 }
 
 bool UInterchangeManager::ImportAsset(const FString& ContentPath, const UInterchangeSourceData* SourceData, const FImportAssetParameters& ImportAssetParameters)
+{
+	return ImportAssetAsync( ContentPath, SourceData, ImportAssetParameters ).IsValid();
+}
+
+Interchange::FAsyncImportResult UInterchangeManager::ImportAssetAsync(const FString& ContentPath, const UInterchangeSourceData* SourceData, const FImportAssetParameters& ImportAssetParameters)
 {
 	FString PackageBasePath = ContentPath;
 	if(!ImportAssetParameters.ReimportAsset)
@@ -266,19 +303,25 @@ bool UInterchangeManager::ImportAsset(const FString& ContentPath, const UInterch
 		check(AsyncHelper->BaseNodeContainerAdapters[SourceDataIndex].IsValid());
 		AsyncHelper->BaseNodeContainerAdapters[SourceDataIndex].Get()->SetBaseNodeContainer(&AsyncHelper->BaseNodeContainers[SourceDataIndex]);
 	}
-	
 
-	//Get all pipeline candidate we want for this import
-	TArray<UClass*> PipelineCandidates;
-	FindPipelineCandidate(PipelineCandidates);
+	if ( ImportAssetParameters.OverridePipeline == nullptr )
+	{
+		//Get all pipeline candidate we want for this import
+		/*TArray<UClass*> PipelineCandidates;
+		FindPipelineCandidate(PipelineCandidates);
 
-	// Stack all pipelines, for this import proto. TODO: We need to be able to control which pipeline we use for the import
-	// It can be set in the project settings and can also be set by a UI where the user create the pipeline stack he want.
-	// This should be a list of available pipelines that can be drop into a stack where you can control the order.
-	for (int32 GraphPipelineNumber = 0; GraphPipelineNumber < PipelineCandidates.Num(); ++GraphPipelineNumber)
+		// Stack all pipelines, for this import proto. TODO: We need to be able to control which pipeline we use for the import
+		// It can be set in the project settings and can also be set by a UI where the user create the pipeline stack he want.
+		// This should be a list of available pipelines that can be drop into a stack where you can control the order.
+		for (int32 GraphPipelineNumber = 0; GraphPipelineNumber < PipelineCandidates.Num(); ++GraphPipelineNumber)
+		{
+			UInterchangePipelineBase* GeneratedPipeline = NewObject<UInterchangePipelineBase>(GetTransientPackage(), PipelineCandidates[GraphPipelineNumber], NAME_None, RF_NoFlags);
+			AsyncHelper->Pipelines.Add(GeneratedPipeline);
+		}*/
+	}
+	else
 	{
-		UInterchangePipelineBase* GeneratedPipeline = NewObject<UInterchangePipelineBase>(GetTransientPackage(), PipelineCandidates[GraphPipelineNumber], NAME_None, RF_NoFlags);
-		AsyncHelper->Pipelines.Add(GeneratedPipeline);
+		AsyncHelper->Pipelines.Add(ImportAssetParameters.OverridePipeline);
 	}
 
 	//Create/Start import tasks
@@ -317,7 +360,7 @@ bool UInterchangeManager::ImportAsset(const FString& ContentPath, const UInterch
 
 	//The graph parsing task will create the FCreateAssetTask that will run after them, the FAssetImportTask will call the appropriate Post asset import pipeline when the asset is completed
 
-	return true;
+	return Interchange::FAsyncImportResult{ AsyncHelper->RootObject.GetFuture(), AsyncHelper->RootObjectCompletionEvent };
 }
 
 bool UInterchangeManager::ImportScene(const FString& ImportContext, const UInterchangeSourceData* SourceData, bool bIsReimport, bool bIsAutomated)
diff --git a/Engine/Source/Runtime/Interchange/Core/Private/Tasks/InterchangeTaskCompletion.cpp b/Engine/Source/Runtime/Interchange/Core/Private/Tasks/InterchangeTaskCompletion.cpp
--- a/Engine/Source/Runtime/Interchange/Core/Private/Tasks/InterchangeTaskCompletion.cpp
+++ b/Engine/Source/Runtime/Interchange/Core/Private/Tasks/InterchangeTaskCompletion.cpp
@@ -57,6 +57,12 @@ void Interchange::FTaskCompletion::DoTask(ENamedThreads::Type CurrentThread, con
 			//Notify the asset registry
 			FAssetRegistryModule::AssetCreated(Asset);
 #endif //WITH_ENGINE
+
+			if (SourceIndex == 0)
+			{
+				AsyncHelper->RootObject.SetValue(Asset);
+				AsyncHelper->RootObjectCompletionEvent->DispatchSubsequents();
+			}
 		}
 	}
 	//Release the async helper
diff --git a/Engine/Source/Runtime/Interchange/Core/Public/InterchangeManager.h b/Engine/Source/Runtime/Interchange/Core/Public/InterchangeManager.h
--- a/Engine/Source/Runtime/Interchange/Core/Public/InterchangeManager.h
+++ b/Engine/Source/Runtime/Interchange/Core/Public/InterchangeManager.h
@@ -63,6 +63,8 @@ namespace Interchange
 	class FImportAsyncHelper : public FGCObject
 	{
 	public:
+		FImportAsyncHelper();
+
 		~FImportAsyncHelper()
 		{
 			CleanUp();
@@ -100,9 +102,27 @@ namespace Interchange
 
 		FImportAsyncHelperData TaskData;
 
+		TPromise< UObject* > RootObject;
+		FGraphEventRef RootObjectCompletionEvent;
+
 		void CleanUp();
 	};
 
+	class INTERCHANGECORE_API FAsyncImportResult
+	{
+	public:
+		FAsyncImportResult() = default;
+		explicit FAsyncImportResult( TFuture< UObject* >&& InFutureObject, const FGraphEventRef& InGraphEvent );
+
+		bool IsValid() const;
+		UObject* Get() const;
+		FAsyncImportResult Next( TFunction< UObject*( UObject* ) > Continuation );
+
+	private:
+		TFuture< UObject* > FutureObject;
+		FGraphEventRef GraphEvent;
+	};
+
 	void SanitizeInvalidChar(FString& String);
 
 } //End namespace interchange
@@ -117,6 +137,9 @@ struct FImportAssetParameters
 
 	UPROPERTY(EditAnywhere, BlueprintReadWrite, Category = "Interchange|ImportAsset")
 	bool bIsAutomated = false;
+
+	UPROPERTY(EditAnywhere, BlueprintReadWrite, Category = "Interchange|ImportAsset")
+	UInterchangePipelineBase* OverridePipeline = nullptr;
 };
 
 UCLASS(Transient, BlueprintType)
@@ -210,6 +233,7 @@ public:
 	 */
 	UFUNCTION(BlueprintCallable, Category = "Interchange | Import Manager")
 	bool ImportAsset(const FString& ContentPath, const UInterchangeSourceData* SourceData, const FImportAssetParameters& ImportAssetParameters);
+	Interchange::FAsyncImportResult ImportAssetAsync(const FString& ContentPath, const UInterchangeSourceData* SourceData, const FImportAssetParameters& ImportAssetParameters);
 
 	/**
 	 * Call this to start an import scene process, the caller must specify a source data.
diff --git a/Engine/Source/Runtime/Interchange/Core/Public/Nodes/InterchangeBaseNodeContainerAdapter.h b/Engine/Source/Runtime/Interchange/Core/Public/Nodes/InterchangeBaseNodeContainerAdapter.h
--- a/Engine/Source/Runtime/Interchange/Core/Public/Nodes/InterchangeBaseNodeContainerAdapter.h
+++ b/Engine/Source/Runtime/Interchange/Core/Public/Nodes/InterchangeBaseNodeContainerAdapter.h
@@ -373,6 +373,6 @@ public:
 	UFUNCTION(BlueprintCallable, Category = "Interchange | Node Attribute")
 	bool SetNodeAttributeValueVector4(const FName NodeIdentity, const FName AttributeName, const FVector4& Value);
 
-private:
+public:
 	Interchange::FBaseNodeContainer* BaseNodeContainer = nullptr;
 };
