Fixed Interchange breaking its promise when we failed to import anything

diff --git a/Engine/Plugins/Enterprise/DatasmithImporter/Source/DatasmithImporter/Public/Utility/DatasmithImporterUtils.h b/Engine/Plugins/Enterprise/DatasmithImporter/Source/DatasmithImporter/Public/Utility/DatasmithImporterUtils.h
--- a/Engine/Plugins/Enterprise/DatasmithImporter/Source/DatasmithImporter/Public/Utility/DatasmithImporterUtils.h
+++ b/Engine/Plugins/Enterprise/DatasmithImporter/Source/DatasmithImporter/Public/Utility/DatasmithImporterUtils.h
@@ -450,16 +450,12 @@ inline UTexture* FDatasmithImporterUtils::FindAsset< UTexture >( const FDatasmit
 			// Check if the AssetsMap is already tracking our asset
 			if ( AssetsMap && AssetsMap->Contains( ObjectPathName ) )
 			{
-				UTexture* Texture = (*AssetsMap)[ FName( ObjectPathName ) ].LoadSynchronous();
-				ensure( Texture );
-				return Texture;
+				return (*AssetsMap)[ FName( ObjectPathName ) ].LoadSynchronous();
 			}
 			else
 			{
 				UPackage* FinalPackage = FDatasmithFindAssetTypeHelper< UTexture >::GetFinalPackage( AssetsContext );
-				UTexture* Texture = FindObject< UTexture >( FinalPackage, ObjectPathName );
-				ensure( Texture );
-				return Texture;
+				return FindObject< UTexture >( FinalPackage, ObjectPathName );
 			}
 		}
 	}
diff --git a/Engine/Source/Runtime/Interchange/Core/Private/Tasks/InterchangeTaskCompletion.cpp b/Engine/Source/Runtime/Interchange/Core/Private/Tasks/InterchangeTaskCompletion.cpp
--- a/Engine/Source/Runtime/Interchange/Core/Private/Tasks/InterchangeTaskCompletion.cpp
+++ b/Engine/Source/Runtime/Interchange/Core/Private/Tasks/InterchangeTaskCompletion.cpp
@@ -65,7 +65,14 @@ void Interchange::FTaskCompletion::DoTask(ENamedThreads::Type CurrentThread, con
 			}
 		}
 	}
+
+	if (!AsyncHelper->RootObjectCompletionEvent->IsComplete())
+	{
+		AsyncHelper->RootObject.SetValue(nullptr);
+		AsyncHelper->RootObjectCompletionEvent->DispatchSubsequents();
+	}
+
 	//Release the async helper
 	AsyncHelper = nullptr;
 	InterchangeManager->ReleaseAsyncHelper(WeakAsyncHelper);
-}
\ No newline at end of file
+}
