fix datasmith dependency to use editor plugin for interchange

diff --git a/Engine/Plugins/Enterprise/DatasmithImporter/DatasmithImporter.uplugin b/Engine/Plugins/Enterprise/DatasmithImporter/DatasmithImporter.uplugin
--- a/Engine/Plugins/Enterprise/DatasmithImporter/DatasmithImporter.uplugin
+++ b/Engine/Plugins/Enterprise/DatasmithImporter/DatasmithImporter.uplugin
@@ -45,7 +45,7 @@
 			"Enabled": true
 		},
 		{
-			"Name": "Interchange",
+			"Name": "InterchangeEditor",
 			"Enabled": true
 		},
 		{
