Remove automated scene thumbnail capture after a datasmith import
  - Was recently made obsolete by async static mesh causing the thumbnail to be empty
  - A user can manually capture a better thumbnail for his level at his own convenience

diff --git a/Engine/Plugins/Enterprise/DatasmithImporter/Source/DatasmithImporter/Private/DatasmithImportFactory.cpp b/Engine/Plugins/Enterprise/DatasmithImporter/Source/DatasmithImporter/Private/DatasmithImportFactory.cpp
--- a/Engine/Plugins/Enterprise/DatasmithImporter/Source/DatasmithImporter/Private/DatasmithImportFactory.cpp
+++ b/Engine/Plugins/Enterprise/DatasmithImporter/Source/DatasmithImporter/Private/DatasmithImportFactory.cpp
@@ -90,21 +90,21 @@ namespace DatasmithImportFactoryImpl
 		return nullptr;
 	}
 
-	static void CaptureSceneThumbnail(FDatasmithImportContext& InContext)
+	static void SetupSceneViewport(FDatasmithImportContext& InContext)
 	{
 		if ( !InContext.ShouldImportActors() || !InContext.SceneAsset || InContext.ActorsContext.FinalSceneActors.Num() == 0)
 		{
 			return;
 		}
 
-		TRACE_CPUPROFILER_EVENT_SCOPE(DatasmithImportFactoryImpl::CaptureSceneThumbnail);
+		TRACE_CPUPROFILER_EVENT_SCOPE(DatasmithImportFactoryImpl::SetupSceneViewport);
 
 		// Use the first scene actor for the thumbnail
 		ADatasmithSceneActor* SceneActor = *InContext.ActorsContext.FinalSceneActors.CreateIterator();
 
 		TArray< FAssetData > AssetDataList;
 		AssetDataList.Add( FAssetData( InContext.SceneAsset ) );
-		DatasmithImportFactoryHelper::CaptureSceneThumbnail(SceneActor, AssetDataList);
+		DatasmithImportFactoryHelper::SetupSceneViewport(SceneActor, AssetDataList);
 	}
 
 	static bool CreateSceneAsset( FDatasmithImportContext& InContext )
@@ -293,9 +293,8 @@ namespace DatasmithImportFactoryImpl
 		}
 		FDatasmithImporter::FinalizeImport(InContext, TSet<UObject*>());
 
-		// THUMBNAIL
 		// Must be called after the actors are spawned since we will compute the scene bounds
-		DatasmithImportFactoryImpl::CaptureSceneThumbnail( InContext );
+		DatasmithImportFactoryImpl::SetupSceneViewport( InContext );
 
 		return true;
 	}
diff --git a/Engine/Plugins/Enterprise/DatasmithImporter/Source/DatasmithImporter/Private/Utility/DatasmithImportFactoryHelper.cpp b/Engine/Plugins/Enterprise/DatasmithImporter/Source/DatasmithImporter/Private/Utility/DatasmithImportFactoryHelper.cpp
--- a/Engine/Plugins/Enterprise/DatasmithImporter/Source/DatasmithImporter/Private/Utility/DatasmithImportFactoryHelper.cpp
+++ b/Engine/Plugins/Enterprise/DatasmithImporter/Source/DatasmithImporter/Private/Utility/DatasmithImportFactoryHelper.cpp
@@ -69,7 +69,7 @@ namespace DatasmithImportFactoryHelper
 		}
 	}
 
-	void CaptureSceneThumbnail(AActor* SceneActor, TArray<FAssetData>& AssetDataList)
+	void SetupSceneViewport(AActor* SceneActor, TArray<FAssetData>& AssetDataList)
 	{
 		if ( !GCurrentLevelEditingViewportClient || !GEditor->GetActiveViewport())
 		{
@@ -98,11 +98,6 @@ namespace DatasmithImportFactoryHelper
 		GCurrentLevelEditingViewportClient->SetGameView( true );
 		GCurrentLevelEditingViewportClient->SetIsCameraCut();
 
-		GEditor->GetActiveViewport()->Draw();
-
-		IContentBrowserSingleton& ContentBrowser = FModuleManager::LoadModuleChecked<FContentBrowserModule>("ContentBrowser").Get();
-		ContentBrowser.CaptureThumbnailFromViewport( GEditor->GetActiveViewport(), AssetDataList );
-
 		// Unlock the viewport if we locked it to capture the thumbnail
 		if ( OriginalActorLock )
 		{
diff --git a/Engine/Plugins/Enterprise/DatasmithImporter/Source/DatasmithImporter/Private/Utility/DatasmithImportFactoryHelper.h b/Engine/Plugins/Enterprise/DatasmithImporter/Source/DatasmithImporter/Private/Utility/DatasmithImportFactoryHelper.h
--- a/Engine/Plugins/Enterprise/DatasmithImporter/Source/DatasmithImporter/Private/Utility/DatasmithImportFactoryHelper.h
+++ b/Engine/Plugins/Enterprise/DatasmithImporter/Source/DatasmithImporter/Private/Utility/DatasmithImportFactoryHelper.h
@@ -17,5 +17,5 @@ namespace DatasmithImportFactoryHelper
 
 	void ComputeBounds( USceneComponent* ActorComponent, FBox& Bounds );
 
-	void CaptureSceneThumbnail(AActor* SceneActor, TArray<FAssetData>& AssetDataList);
+	void SetupSceneViewport(AActor* SceneActor, TArray<FAssetData>& AssetDataList);
 }
