Optimize datasmith import for huge scene
  - Replace component deregistration/registration by a proper render state recreation context
  - Use material update context at higher level so that all material updates are put into a single scope

  - 6m47s to 5m12s for DATASET-0034a

diff --git a/Engine/Plugins/Enterprise/DatasmithImporter/Source/DatasmithImporter/Private/DatasmithImporter.cpp b/Engine/Plugins/Enterprise/DatasmithImporter/Source/DatasmithImporter/Private/DatasmithImporter.cpp
--- a/Engine/Plugins/Enterprise/DatasmithImporter/Source/DatasmithImporter/Private/DatasmithImporter.cpp
+++ b/Engine/Plugins/Enterprise/DatasmithImporter/Source/DatasmithImporter/Private/DatasmithImporter.cpp
@@ -584,7 +584,7 @@ UMaterialInterface* FDatasmithImporter::ImportMaterial( FDatasmithImportContext&
 	return ImportedMaterial;
 }
 
-UObject* FDatasmithImporter::FinalizeMaterial( UObject* SourceMaterial, const TCHAR* MaterialFolderPath, const TCHAR* TransientPackagePath, const TCHAR* RootFolderPath, UMaterialInterface* ExistingMaterial, TMap< UObject*, UObject* >* ReferencesToRemap)
+UObject* FDatasmithImporter::FinalizeMaterial( UObject* SourceMaterial, const TCHAR* MaterialFolderPath, const TCHAR* TransientPackagePath, const TCHAR* RootFolderPath, UMaterialInterface* ExistingMaterial, TMap< UObject*, UObject* >* ReferencesToRemap, FMaterialUpdateContext* MaterialUpdateContext)
 {
 	TRACE_CPUPROFILER_EVENT_SCOPE(FDatasmithImporter::FinalizeMaterial);
 
@@ -622,7 +622,7 @@ UObject* FDatasmithImporter::FinalizeMaterial( UObject* SourceMaterial, const TC
 
 	UObject* DestinationMaterial = FDatasmithImporterImpl::FinalizeAsset( SourceMaterial, MaterialFolderPath, ExistingMaterial, ReferencesToRemap );
 
-	FDatasmithImporterImpl::CompileMaterial( DestinationMaterial );
+	FDatasmithImporterImpl::CompileMaterial( DestinationMaterial, MaterialUpdateContext);
 
 	return DestinationMaterial;
 }
@@ -1551,119 +1551,107 @@ void FDatasmithImporter::FinalizeImport(FDatasmithImportContext& ImportContext,
 		FDatasmithImporterImpl::CheckAssetPersistenceValidity(ExistingTexture, ImportContext);
 	}
 
-	// Unregister all actors component to avoid excessive refresh in the 3D engine while updating materials.
-	for (TObjectIterator<AActor> ActorIterator; ActorIterator; ++ActorIterator)
 	{
-		if (ActorIterator->GetWorld())
-		{
-			ActorIterator->UnregisterAllComponents( /* bForReregister = */true);
-		}
-	}
-
-	for (const TPair< TSharedRef< IDatasmithBaseMaterialElement >, UMaterialFunction* >& ImportedMaterialFunctionPair : ImportContext.ImportedMaterialFunctions)
-	{
-		if (ImportContext.bUserCancelled)
-		{
-			break;
-		}
-
-		UMaterialFunction* SourceMaterialFunction = ImportedMaterialFunctionPair.Value;
-
-		if (!SourceMaterialFunction || (ValidAssets.Num() > 0 && !ValidAssets.Contains(Cast<UObject>(SourceMaterialFunction))))
-		{
-			continue;
-		}
-
-		FName MaterialFunctionId = ImportedMaterialFunctionPair.Key->GetName();
-		FDatasmithImporterImpl::ReportProgress(Progress, 1.f, FText::FromString(FString::Printf(TEXT("Finalizing assets %d/%d (Material Function %s) ..."), ++AssetIndex, NumAssetsToFinalize, *SourceMaterialFunction->GetName())));
-
-		UMaterialFunction* ExistingMaterialFunction = ImportContext.SceneAsset ? ImportContext.SceneAsset->MaterialFunctions.FindOrAdd(MaterialFunctionId).Get() : nullptr;
-
-		FString SourcePackagePath = SourceMaterialFunction->GetOutermost()->GetName();
-		FString DestinationPackagePath = SourcePackagePath.Replace(*TransientFolderPath, *RootFolderPath, ESearchCase::CaseSensitive);
+		// Destroy render states of all actors to avoid excessive refresh in the renderer while updating materials.
+		// Also batch all the material update together to improve performance.
+		FMaterialUpdateContext MaterialUpdateContext;
 
-		ExistingMaterialFunction = FinalizeMaterialFunction(SourceMaterialFunction, *DestinationPackagePath, ExistingMaterialFunction, &ReferencesToRemap);
-		if (ImportContext.SceneAsset)
+		for (const TPair< TSharedRef< IDatasmithBaseMaterialElement >, UMaterialFunction* >& ImportedMaterialFunctionPair : ImportContext.ImportedMaterialFunctions)
 		{
-			ImportContext.SceneAsset->MaterialFunctions[MaterialFunctionId] = ExistingMaterialFunction;
+			if (ImportContext.bUserCancelled)
+			{
+				break;
+			}
+  
+			UMaterialFunction* SourceMaterialFunction = ImportedMaterialFunctionPair.Value;
+  
+			if (!SourceMaterialFunction || (ValidAssets.Num() > 0 && !ValidAssets.Contains(Cast<UObject>(SourceMaterialFunction))))
+			{
+				continue;
+			}
+  
+			FName MaterialFunctionId = ImportedMaterialFunctionPair.Key->GetName();
+			FDatasmithImporterImpl::ReportProgress(Progress, 1.f, FText::FromString(FString::Printf(TEXT("Finalizing assets %d/%d (Material Function %s) ..."), ++AssetIndex, NumAssetsToFinalize, *SourceMaterialFunction->GetName())));
+  
+			UMaterialFunction* ExistingMaterialFunction = ImportContext.SceneAsset ? ImportContext.SceneAsset->MaterialFunctions.FindOrAdd(MaterialFunctionId).Get() : nullptr;
+  
+			FString SourcePackagePath = SourceMaterialFunction->GetOutermost()->GetName();
+			FString DestinationPackagePath = SourcePackagePath.Replace(*TransientFolderPath, *RootFolderPath, ESearchCase::CaseSensitive);
+  
+			ExistingMaterialFunction = FinalizeMaterialFunction(SourceMaterialFunction, *DestinationPackagePath, ExistingMaterialFunction, &ReferencesToRemap);
+			if (ImportContext.SceneAsset)
+			{
+				ImportContext.SceneAsset->MaterialFunctions[MaterialFunctionId] = ExistingMaterialFunction;
+			}
+  
+			FDatasmithImporterImpl::CheckAssetPersistenceValidity(ExistingMaterialFunction, ImportContext);
 		}
 
-		FDatasmithImporterImpl::CheckAssetPersistenceValidity(ExistingMaterialFunction, ImportContext);
-	}
-
-	TArray<UMaterial*> MaterialsToRefreshAfterVirtualTextureConversion;
-	for (const TPair< TSharedRef< IDatasmithBaseMaterialElement >, UMaterialInterface* >& ImportedMaterialPair : ImportContext.ImportedMaterials)
-	{
-		if (ImportContext.bUserCancelled)
+		TArray<UMaterial*> MaterialsToRefreshAfterVirtualTextureConversion;
+		for (const TPair< TSharedRef< IDatasmithBaseMaterialElement >, UMaterialInterface* >& ImportedMaterialPair : ImportContext.ImportedMaterials)
 		{
-			break;
-		}
+			if (ImportContext.bUserCancelled)
+			{
+				break;
+			}
 
-		UMaterialInterface* SourceMaterialInterface = ImportedMaterialPair.Value;
+			UMaterialInterface* SourceMaterialInterface = ImportedMaterialPair.Value;
 
-		if (!SourceMaterialInterface || (ValidAssets.Num() > 0 && !ValidAssets.Contains(Cast<UObject>(SourceMaterialInterface))))
-		{
-			continue;
-		}
-
-		FName MaterialId = ImportedMaterialPair.Key->GetName();
+			if (!SourceMaterialInterface || (ValidAssets.Num() > 0 && !ValidAssets.Contains(Cast<UObject>(SourceMaterialInterface))))
+			{
+				continue;
+			}
 
-		FDatasmithImporterImpl::ReportProgress(Progress, 1.f, FText::FromString(FString::Printf(TEXT("Finalizing assets %d/%d (Material %s) ..."), ++AssetIndex, NumAssetsToFinalize, *SourceMaterialInterface->GetName())));
+			FName MaterialId = ImportedMaterialPair.Key->GetName();
 
-		UMaterialInterface* ExistingMaterial = ImportContext.SceneAsset ? ImportContext.SceneAsset->Materials.FindOrAdd(MaterialId).Get() : NULL;
+			FDatasmithImporterImpl::ReportProgress(Progress, 1.f, FText::FromString(FString::Printf(TEXT("Finalizing assets %d/%d (Material %s) ..."), ++AssetIndex, NumAssetsToFinalize, *SourceMaterialInterface->GetName())));
 
-		FString SourcePackagePath = SourceMaterialInterface->GetOutermost()->GetName();
-		FString DestinationPackagePath = SourcePackagePath.Replace( *TransientFolderPath, *RootFolderPath, ESearchCase::CaseSensitive );
+			UMaterialInterface* ExistingMaterial = ImportContext.SceneAsset ? ImportContext.SceneAsset->Materials.FindOrAdd(MaterialId).Get() : NULL;
 
-		if (UMaterial* SourceMaterial = Cast< UMaterial >(SourceMaterialInterface))
-		{
-			SourceMaterial->UpdateCachedExpressionData();
+			FString SourcePackagePath = SourceMaterialInterface->GetOutermost()->GetName();
+			FString DestinationPackagePath = SourcePackagePath.Replace( *TransientFolderPath, *RootFolderPath, ESearchCase::CaseSensitive );
 
-			TArray<UObject*> ReferencedTextures;
-			ReferencedTextures = SourceMaterial->GetReferencedTextures();
-			for (UTexture2D* VirtualTexture : ImportContext.AssetsContext.VirtualTexturesToConvert)
+			if (UMaterial* SourceMaterial = Cast< UMaterial >(SourceMaterialInterface))
 			{
-				if (ReferencedTextures.Contains(VirtualTexture))
+				SourceMaterial->UpdateCachedExpressionData();
+
+				TArray<UObject*> ReferencedTextures;
+				ReferencedTextures = SourceMaterial->GetReferencedTextures();
+				for (UTexture2D* VirtualTexture : ImportContext.AssetsContext.VirtualTexturesToConvert)
 				{
-					MaterialsToRefreshAfterVirtualTextureConversion.Add(SourceMaterial);
-					break;
+					if (ReferencedTextures.Contains(VirtualTexture))
+					{
+						MaterialsToRefreshAfterVirtualTextureConversion.Add(SourceMaterial);
+						break;
+					}
 				}
-			}
 
-			for (const FMaterialFunctionInfo& MaterialFunctionInfo : SourceMaterial->GetCachedExpressionData().FunctionInfos)
-			{
-				if (MaterialFunctionInfo.Function && MaterialFunctionInfo.Function->GetOutermost() == SourceMaterial->GetOutermost())
+				for (const FMaterialFunctionInfo& MaterialFunctionInfo : SourceMaterial->GetCachedExpressionData().FunctionInfos)
 				{
-					FinalizeMaterial(MaterialFunctionInfo.Function, *DestinationPackagePath, *TransientFolderPath, *RootFolderPath, nullptr, &ReferencesToRemap);
+					if (MaterialFunctionInfo.Function && MaterialFunctionInfo.Function->GetOutermost() == SourceMaterial->GetOutermost())
+					{
+						FinalizeMaterial(MaterialFunctionInfo.Function, *DestinationPackagePath, *TransientFolderPath, *RootFolderPath, nullptr, &ReferencesToRemap, &MaterialUpdateContext);
+					}
 				}
 			}
-		}
 
-		ExistingMaterial = Cast<UMaterialInterface>( FinalizeMaterial(SourceMaterialInterface, *DestinationPackagePath, *TransientFolderPath, *RootFolderPath, ExistingMaterial, &ReferencesToRemap) );
-		if (ImportContext.SceneAsset)
-		{
-			ImportContext.SceneAsset->Materials[MaterialId] = ExistingMaterial;
-		}
+			ExistingMaterial = Cast<UMaterialInterface>( FinalizeMaterial(SourceMaterialInterface, *DestinationPackagePath, *TransientFolderPath, *RootFolderPath, ExistingMaterial, &ReferencesToRemap, &MaterialUpdateContext) );
+			if (ImportContext.SceneAsset)
+			{
+				ImportContext.SceneAsset->Materials[MaterialId] = ExistingMaterial;
+			}
 
-		// Add material to array of packages to apply soft object path redirection to
-		if (ExistingMaterial)
-		{
-			PackagesToCheck.Add(ExistingMaterial->GetOutermost());
-			FDatasmithImporterImpl::CheckAssetPersistenceValidity(ExistingMaterial, ImportContext);
+			// Add material to array of packages to apply soft object path redirection to
+			if (ExistingMaterial)
+			{
+				PackagesToCheck.Add(ExistingMaterial->GetOutermost());
+				FDatasmithImporterImpl::CheckAssetPersistenceValidity(ExistingMaterial, ImportContext);
+			}
 		}
 	}
 
 	FDatasmithImporterImpl::ConvertUnsupportedVirtualTexture(ImportContext, ImportContext.AssetsContext.VirtualTexturesToConvert, ReferencesToRemap);
 
-	// Materials have been updated, we can register everything back.
-	for (TObjectIterator<AActor> ActorIterator; ActorIterator; ++ActorIterator)
-	{
-		if (ActorIterator->GetWorld())
-		{
-			ActorIterator->RegisterAllComponents();
-		}
-	}
-
 	// Sometimes, the data is invalid and we get the same UStaticMesh multiple times
 	TSet< UStaticMesh* > StaticMeshes;
 	for (TPair< TSharedRef< IDatasmithMeshElement >, UStaticMesh* >& ImportedStaticMeshPair : ImportContext.ImportedStaticMeshes)
diff --git a/Engine/Plugins/Enterprise/DatasmithImporter/Source/DatasmithImporter/Private/Utility/DatasmithImporterImpl.cpp b/Engine/Plugins/Enterprise/DatasmithImporter/Source/DatasmithImporter/Private/Utility/DatasmithImporterImpl.cpp
--- a/Engine/Plugins/Enterprise/DatasmithImporter/Source/DatasmithImporter/Private/Utility/DatasmithImporterImpl.cpp
+++ b/Engine/Plugins/Enterprise/DatasmithImporter/Source/DatasmithImporter/Private/Utility/DatasmithImporterImpl.cpp
@@ -307,7 +307,7 @@ void FDatasmithImporterImpl::SetTexturesMode( FDatasmithImportContext& ImportCon
 	}
 }
 
-void FDatasmithImporterImpl::CompileMaterial( UObject* Material )
+void FDatasmithImporterImpl::CompileMaterial( UObject* Material, FMaterialUpdateContext* MaterialUpdateContext)
 {
 	if ( !Material->IsA< UMaterialInterface >() && !Material->IsA< UMaterialFunctionInterface >() )
 	{
@@ -316,11 +316,16 @@ void FDatasmithImporterImpl::CompileMaterial( UObject* Material )
 
 	TRACE_CPUPROFILER_EVENT_SCOPE(FDatasmithImporterImpl::CompileMaterial);
 
-	FMaterialUpdateContext MaterialUpdateContext;
+	TUniquePtr<FMaterialUpdateContext> LocalMaterialUpdateContext;
+	if (MaterialUpdateContext == nullptr)
+	{
+		LocalMaterialUpdateContext = MakeUnique<FMaterialUpdateContext>();
+		MaterialUpdateContext = LocalMaterialUpdateContext.Get();
+	}
 
 	if ( UMaterialInterface* MaterialInterface = Cast< UMaterialInterface >( Material ) )
 	{
-		MaterialUpdateContext.AddMaterialInterface( MaterialInterface );
+		MaterialUpdateContext->AddMaterialInterface( MaterialInterface );
 	}
 
 	if ( UMaterialInstanceConstant* ConstantMaterialInstance = Cast< UMaterialInstanceConstant >( Material ) )
diff --git a/Engine/Plugins/Enterprise/DatasmithImporter/Source/DatasmithImporter/Private/Utility/DatasmithImporterImpl.h b/Engine/Plugins/Enterprise/DatasmithImporter/Source/DatasmithImporter/Private/Utility/DatasmithImporterImpl.h
--- a/Engine/Plugins/Enterprise/DatasmithImporter/Source/DatasmithImporter/Private/Utility/DatasmithImporterImpl.h
+++ b/Engine/Plugins/Enterprise/DatasmithImporter/Source/DatasmithImporter/Private/Utility/DatasmithImporterImpl.h
@@ -39,6 +39,7 @@ class FFeedbackContext;
 class IDatasmithBaseMaterialElement;
 class UActorComponent;
 class UDatasmithObjectTemplate;
+class FMaterialUpdateContext;
 class UMaterial;
 class UMaterialInterface;
 class USceneComponent;
@@ -73,7 +74,7 @@ public:
 	/** Set the texture mode on each texture element based on its usage in the materials */
 	static void SetTexturesMode( FDatasmithImportContext& ImportContext );
 
-	static void CompileMaterial( UObject* Material );
+	static void CompileMaterial( UObject* Material, FMaterialUpdateContext* MaterialUpdateContext = nullptr);
 
 	static void FixReferencesForObject( UObject* Object, const TMap< UObject*, UObject* >& ReferencesToRemap );
 
diff --git a/Engine/Plugins/Enterprise/DatasmithImporter/Source/DatasmithImporter/Public/DatasmithImporter.h b/Engine/Plugins/Enterprise/DatasmithImporter/Source/DatasmithImporter/Public/DatasmithImporter.h
--- a/Engine/Plugins/Enterprise/DatasmithImporter/Source/DatasmithImporter/Public/DatasmithImporter.h
+++ b/Engine/Plugins/Enterprise/DatasmithImporter/Source/DatasmithImporter/Public/DatasmithImporter.h
@@ -120,10 +120,11 @@ struct DATASMITHIMPORTER_API FDatasmithImporter
 	 * @param MaterialFolderPath	The root folder where we want to import the scene, used to finalize parent materials.
 	 * @param ExistingMaterial		An optional existing material on which we will finalize over.
 	 * @param ReferencesToRemap		A map of migrated objects. If SourceMaterial refers to any key objects, those references will be replaced by the corresponding map value.
+	 * @param MaterialUpdateContext		An optional material update context to batch updates and improve performance.
 	 *
 	 * @return The resulting material
 	 */
-	static UObject* FinalizeMaterial( UObject* SourceMaterial, const TCHAR* MaterialFolderPath, const TCHAR* TransientFolderPath, const TCHAR* RootFolderPath, UMaterialInterface* ExistingMaterial, TMap< UObject*, UObject* >* ReferencesToRemap = nullptr );
+	static UObject* FinalizeMaterial( UObject* SourceMaterial, const TCHAR* MaterialFolderPath, const TCHAR* TransientFolderPath, const TCHAR* RootFolderPath, UMaterialInterface* ExistingMaterial, TMap< UObject*, UObject* >* ReferencesToRemap = nullptr, FMaterialUpdateContext* MaterialUpdateContext = nullptr);
 
 	/**
 	 * Imports the actors
