Fixed missing ies textures when importing Datasmith scenes
   - Root cause: Asynchronous creation of IES texture is not supported
   - Solution: Skip it for the time being. Will be addressed later on.

diff --git a/Engine/Plugins/Enterprise/DatasmithImporter/Source/DatasmithImporter/Private/DatasmithImporter.cpp b/Engine/Plugins/Enterprise/DatasmithImporter/Source/DatasmithImporter/Private/DatasmithImporter.cpp
--- a/Engine/Plugins/Enterprise/DatasmithImporter/Source/DatasmithImporter/Private/DatasmithImporter.cpp
+++ b/Engine/Plugins/Enterprise/DatasmithImporter/Source/DatasmithImporter/Private/DatasmithImporter.cpp
@@ -342,6 +342,13 @@ void FDatasmithImporter::ImportTextures( FDatasmithImportContext& ImportContext
 				break;
 			}
 
+			// Skip asynchronous creation for IES textures as it is not supported.
+			// The creation of such textures will synchronously happen later in the call to FDatasmithTextureImporter::CreateTexture
+			if (FilteredTextureElements[TextureIndex]->GetTextureMode() == EDatasmithTextureMode::Ies)
+			{
+				continue;
+			}
+
 			UE::Interchange::FAssetImportResultRef FutureTexture = DatasmithTextureImporter.CreateTextureAsync( FilteredTextureElements[TextureIndex] );
 
 			if ( FutureTexture->IsValid() )
