Fixing PVS 7.7 Issues:

warning V595: The 'XYZ' pointer was utilized before it was verified against nullptr. Check lines: #, #.

diff --git a/Engine/Plugins/Enterprise/DatasmithImporter/Source/DatasmithImporter/Private/Utility/DatasmithImporterImpl.cpp b/Engine/Plugins/Enterprise/DatasmithImporter/Source/DatasmithImporter/Private/Utility/DatasmithImporterImpl.cpp
--- a/Engine/Plugins/Enterprise/DatasmithImporter/Source/DatasmithImporter/Private/Utility/DatasmithImporterImpl.cpp
+++ b/Engine/Plugins/Enterprise/DatasmithImporter/Source/DatasmithImporter/Private/Utility/DatasmithImporterImpl.cpp
@@ -451,7 +451,7 @@ UObject* FDatasmithImporterImpl::FinalizeAsset( UObject* SourceAsset, const TCHA
 
 	TArray< FMigratedTemplatePairType > MigratedTemplates = MigrateTemplates( SourceAsset, ExistingAsset, ReferencesToRemap, false );
 
-	UObject* FinalAsset = PublicizeAsset( SourceAsset, AssetPath, ExistingAsset );
+	UObject* FinalAsset = PublicizeAsset( SourceAsset, AssetPath, ExistingAsset ); //-V595
 
 	ApplyMigratedTemplates( MigratedTemplates, FinalAsset );
 
diff --git a/Engine/Source/Developer/ShaderFormatOpenGL/Private/OpenGLShaderCompiler.cpp b/Engine/Source/Developer/ShaderFormatOpenGL/Private/OpenGLShaderCompiler.cpp
--- a/Engine/Source/Developer/ShaderFormatOpenGL/Private/OpenGLShaderCompiler.cpp
+++ b/Engine/Source/Developer/ShaderFormatOpenGL/Private/OpenGLShaderCompiler.cpp
@@ -2454,7 +2454,7 @@ void FOpenGLFrontend::CompileShader(const FShaderCompilerInput& Input, FShaderCo
 #if VALIDATE_GLSL_WITH_DRIVER
 		PrecompileShader(Output, Input, GlslShaderSource, Version, Frequency);
 #else // VALIDATE_GLSL_WITH_DRIVER
-		int32 SourceLen = FCStringAnsi::Strlen(GlslShaderSource);
+		int32 SourceLen = FCStringAnsi::Strlen(GlslShaderSource); //-V595
 		Output.Target = Input.Target;
 		BuildShaderOutput(Output, Input, GlslShaderSource, SourceLen, Version);
 #endif // VALIDATE_GLSL_WITH_DRIVER
diff --git a/Engine/Source/Runtime/Engine/Private/ADPCMAudioInfo.cpp b/Engine/Source/Runtime/Engine/Private/ADPCMAudioInfo.cpp
--- a/Engine/Source/Runtime/Engine/Private/ADPCMAudioInfo.cpp
+++ b/Engine/Source/Runtime/Engine/Private/ADPCMAudioInfo.cpp
@@ -586,7 +586,10 @@ bool FADPCMAudioInfo::StreamCompressedData(uint8* Destination, bool bLooping, ui
 	if (bDecompressorReleased)
 	{
 		UE_LOG(LogAudio, Error, TEXT("Stream Compressed Data called after release!"));
-		FMemory::Memzero(Destination, BufferSize);
+		if (Destination != nullptr)
+		{
+			FMemory::Memzero(Destination, BufferSize);
+		}
 		return true;
 	}
 
