Datasmith: Move the Import button of the level editor (from the toolbar into the content menu)

diff --git a/Engine/Plugins/Enterprise/DatasmithImporter/Source/DatasmithImporter/Private/DatasmithImporterHelper.cpp b/Engine/Plugins/Enterprise/DatasmithImporter/Source/DatasmithImporter/Private/DatasmithImporterHelper.cpp
--- a/Engine/Plugins/Enterprise/DatasmithImporter/Source/DatasmithImporter/Private/DatasmithImporterHelper.cpp
+++ b/Engine/Plugins/Enterprise/DatasmithImporter/Source/DatasmithImporter/Private/DatasmithImporterHelper.cpp
@@ -1,7 +1,6 @@
 // Copyright Epic Games, Inc. All Rights Reserved.
 
 #include "DatasmithImporterHelper.h"
-#include "UI/DatasmithUIManager.h"
 #include "Utility/DatasmithImporterUtils.h"
 
 #include "AssetRegistryModule.h"
@@ -127,7 +126,7 @@ namespace DatasmithImporterHelperInternal
 
 			ImportAssets(InFilenames, Path, Factory);
 		}
-		
+
 		if(GUnrealEd != nullptr)
 		{
 			// We just finished a potentially long operation and we don't want auto save triggered as soon as we finish the import
@@ -182,8 +181,6 @@ void FDatasmithImporterHelper::ImportInternal(UFactory* FactoryCDO)
 	TStrongObjectPtr<UFactory> Factory(NewObject< UFactory >( (UObject*)GetTransientPackage(), FactoryCDO->GetClass() ) );
 	Factory->ConfigureProperties();
 
-	FDatasmithUIManager::Get().SetLastFactoryUsed(Factory->GetClass());
-
 	TArray<FString> OpenedFiles;
 	FString DefaultLocation(FEditorDirectories::Get().GetLastDirectory(ELastDirectory::GENERIC_IMPORT));
 
diff --git a/Engine/Plugins/Enterprise/DatasmithImporter/Source/DatasmithImporter/Private/DatasmithImporterModule.cpp b/Engine/Plugins/Enterprise/DatasmithImporter/Source/DatasmithImporter/Private/DatasmithImporterModule.cpp
--- a/Engine/Plugins/Enterprise/DatasmithImporter/Source/DatasmithImporter/Private/DatasmithImporterModule.cpp
+++ b/Engine/Plugins/Enterprise/DatasmithImporter/Source/DatasmithImporter/Private/DatasmithImporterModule.cpp
@@ -18,8 +18,8 @@
 #include "DatasmithUtils.h"
 #include "ObjectTemplates/DatasmithObjectTemplate.h"
 #include "ObjectTemplates/DatasmithStaticMeshTemplate.h"
-#include "UI/DatasmithUIManager.h"
 #include "UI/DatasmithConsumerDetails.h"
+#include "UI/DatasmithStyle.h"
 
 #include "AssetToolsModule.h"
 #include "ContentBrowserDelegates.h"
@@ -65,7 +65,8 @@ public:
 		// Disable any UI feature if running in command mode
 		if (UToolMenus::IsToolMenuUIEnabled())
 		{
-			FDatasmithUIManager::Initialize();
+			FDatasmithStyle::Initialize();
+			FDatasmithStyle::SetIcon(TEXT("Import"), TEXT("DatasmithImporter/Content/Icons/DatasmithImporterIcon40"));
 
 			SetupMenuEntry();
 			SetupContentBrowserContextMenuExtender();
@@ -103,7 +104,7 @@ public:
 			RemoveLevelEditorContextMenuExtender();
 			RemoveContentBrowserContextMenuExtender();
 
-			FDatasmithUIManager::Shutdown();
+			FDatasmithStyle::Shutdown();
 
 			// Register the details customizer
 			FPropertyEditorModule& PropertyModule = FModuleManager::LoadModuleChecked< FPropertyEditorModule >( TEXT("PropertyEditor") );
@@ -167,13 +168,18 @@ void FDatasmithImporterModule::SetupMenuEntry()
 {
 	if (!IsRunningCommandlet())
 	{
-		FDatasmithUIManager::Get().AddMenuEntry(
+		UToolMenu* ContentMenu = UToolMenus::Get()->ExtendMenu("LevelEditor.LevelEditorToolBar.ContentQuickMenu");
+		check(ContentMenu);
+
+		FToolMenuSection& Section = ContentMenu->FindOrAddSection("ImportAssets");
+		Section.InitSection("ImportAssets", LOCTEXT("ImportAssets_Label", "Import Assets"), FToolMenuInsert());
+
+		Section.AddMenuEntry(
 			TEXT("Import"),
-			LOCTEXT("DatasmithImport", "Datasmith"),
-			LOCTEXT("DatasmithImportTooltip", "Import Unreal Datasmith file"),
-			TEXT("DatasmithImporter/Content/Icons/DatasmithImporterIcon40"),
-			FExecuteAction::CreateRaw(this, &FDatasmithImporterModule::OnClickedMenuEntry),
-			UDatasmithImportFactory::StaticClass()
+			LOCTEXT("DatasmithImport", "Datasmith..."), // label
+			LOCTEXT("DatasmithImportTooltip", "Import Unreal Datasmith file"), // description
+			FSlateIcon(FDatasmithStyle::GetStyleSetName(), TEXT("Datasmith.Import")),
+			FUIAction(FExecuteAction::CreateRaw(this, &FDatasmithImporterModule::OnClickedMenuEntry)/*, FCanExecuteAction()*/)
 		);
 	}
 }
diff --git a/Engine/Plugins/Enterprise/DatasmithImporter/Source/DatasmithImporter/Private/UI/DatasmithUICommands.cpp b/Engine/Plugins/Enterprise/DatasmithImporter/Source/DatasmithImporter/Private/UI/DatasmithUICommands.cpp
deleted file mode 100644
diff --git a/Engine/Plugins/Enterprise/DatasmithImporter/Source/DatasmithImporter/Private/UI/DatasmithUIManager.cpp b/Engine/Plugins/Enterprise/DatasmithImporter/Source/DatasmithImporter/Private/UI/DatasmithUIManager.cpp
deleted file mode 100644
diff --git a/Engine/Plugins/Enterprise/DatasmithImporter/Source/DatasmithImporter/Public/UI/DatasmithUICommands.h b/Engine/Plugins/Enterprise/DatasmithImporter/Source/DatasmithImporter/Public/UI/DatasmithUICommands.h
deleted file mode 100644
diff --git a/Engine/Plugins/Enterprise/DatasmithImporter/Source/DatasmithImporter/Public/UI/DatasmithUIManager.h b/Engine/Plugins/Enterprise/DatasmithImporter/Source/DatasmithImporter/Public/UI/DatasmithUIManager.h
deleted file mode 100644
