Fixed Datasmith build

diff --git a/Engine/Plugins/Enterprise/DatasmithImporter/Source/DatasmithImporter/Private/DatasmithLandscapeImporter.cpp b/Engine/Plugins/Enterprise/DatasmithImporter/Source/DatasmithImporter/Private/DatasmithLandscapeImporter.cpp
--- a/Engine/Plugins/Enterprise/DatasmithImporter/Source/DatasmithImporter/Private/DatasmithLandscapeImporter.cpp
+++ b/Engine/Plugins/Enterprise/DatasmithImporter/Source/DatasmithImporter/Private/DatasmithLandscapeImporter.cpp
@@ -22,10 +22,11 @@
 AActor* FDatasmithLandscapeImporter::ImportLandscapeActor( const TSharedRef< IDatasmithLandscapeElement >& LandscapeActorElement, FDatasmithImportContext& ImportContext, EDatasmithImportActorPolicy ImportActorPolicy )
 {
 	const bool bSingleFile = true;
+	const bool bFlipYAxis = false;
 	FLandscapeImportDescriptor OutImportDescriptor;
 	OutImportDescriptor.Scale = LandscapeActorElement->GetScale();
 	FText OutMessage;
-	ELandscapeImportResult ImportResult = FLandscapeImportHelper::GetHeightmapImportDescriptor(LandscapeActorElement->GetHeightmap(), bSingleFile, OutImportDescriptor, OutMessage);
+	ELandscapeImportResult ImportResult = FLandscapeImportHelper::GetHeightmapImportDescriptor(LandscapeActorElement->GetHeightmap(), bSingleFile, bFlipYAxis, OutImportDescriptor, OutMessage);
 	if (ImportResult == ELandscapeImportResult::Error)
 	{
 		return nullptr;
