Fix compile issue surfaced in Linux editor cross-compile stemming from CL15273466.

diff --git a/Engine/Plugins/Enterprise/DatasmithImporter/Source/DatasmithImporter/Private/Utility/DatasmithImporterImpl.cpp b/Engine/Plugins/Enterprise/DatasmithImporter/Source/DatasmithImporter/Private/Utility/DatasmithImporterImpl.cpp
--- a/Engine/Plugins/Enterprise/DatasmithImporter/Source/DatasmithImporter/Private/Utility/DatasmithImporterImpl.cpp
+++ b/Engine/Plugins/Enterprise/DatasmithImporter/Source/DatasmithImporter/Private/Utility/DatasmithImporterImpl.cpp
@@ -390,7 +390,7 @@ TArray< FDatasmithImporterImpl::FMigratedTemplatePairType > FDatasmithImporterIm
 		return Results;
 	}
 
-	for ( const TPair< TSubclassOf< UDatasmithObjectTemplate >, UDatasmithObjectTemplate* >& SourceTemplatePair : *SourceTemplates )
+	for ( const TPair< TSubclassOf< UDatasmithObjectTemplate >, TObjectPtr<UDatasmithObjectTemplate> >& SourceTemplatePair : *SourceTemplates )
 	{
 		if ( bIsForActor == SourceTemplatePair.Value->bIsActorTemplate )
 		{
