Removed unnecessary MoveTemp from rvalues to make the Mac build more happy.
#rnx

diff --git a/Engine/Plugins/Enterprise/DatasmithImporter/Source/DatasmithImporter/Private/Utility/DatasmithImporterUtils.cpp b/Engine/Plugins/Enterprise/DatasmithImporter/Source/DatasmithImporter/Private/Utility/DatasmithImporterUtils.cpp
--- a/Engine/Plugins/Enterprise/DatasmithImporter/Source/DatasmithImporter/Private/Utility/DatasmithImporterUtils.cpp
+++ b/Engine/Plugins/Enterprise/DatasmithImporter/Source/DatasmithImporter/Private/Utility/DatasmithImporterUtils.cpp
@@ -1220,8 +1220,8 @@ UStaticMesh* FDatasmithImporterUtils::DuplicateStaticMesh(UStaticMesh* SourceSta
 {
 	// Since static mesh can be quite heavy, remove source models for cloning to reduce useless work.
 	// Will be reinserted on the new duplicated asset or restored on the SourceStaticMesh if bIgnoreBulkData is true.
-	TArray<FStaticMeshSourceModel> SourceModels = MoveTemp(SourceStaticMesh->MoveSourceModels());
-	FStaticMeshSourceModel HiResSourceModel = MoveTemp(SourceStaticMesh->MoveHiResSourceModel());
+	TArray<FStaticMeshSourceModel> SourceModels = SourceStaticMesh->MoveSourceModels();
+	FStaticMeshSourceModel HiResSourceModel = SourceStaticMesh->MoveHiResSourceModel();
 
 	// Temporary flag to skip Postload during DuplicateObject
 	SourceStaticMesh->SetFlags(RF_ArchetypeObject);
