Datasmith now respect the fact that interchange can be turned off.

diff --git a/Engine/Plugins/Enterprise/DatasmithImporter/Source/DatasmithImporter/Private/DatasmithImporter.cpp b/Engine/Plugins/Enterprise/DatasmithImporter/Source/DatasmithImporter/Private/DatasmithImporter.cpp
--- a/Engine/Plugins/Enterprise/DatasmithImporter/Source/DatasmithImporter/Private/DatasmithImporter.cpp
+++ b/Engine/Plugins/Enterprise/DatasmithImporter/Source/DatasmithImporter/Private/DatasmithImporter.cpp
@@ -500,34 +500,37 @@ void FDatasmithImporter::ImportTextures( FDatasmithImportContext& ImportContext
 		FDatasmithTextureResize::Initialize();
 
 		// Try to import textures async first, if possible 
-		for ( int32 TextureIndex = 0; TextureIndex < FilteredTextureElements.Num(); TextureIndex++ )
+		if ( GetDefault<UEditorExperimentalSettings>()->bEnableInterchangeFramework )
 		{
-			ImportContext.bUserCancelled |= FDatasmithImporterImpl::HasUserCancelledTask( ImportContext.FeedbackContext );
-
-			if (ImportContext.bUserCancelled)
+			for ( int32 TextureIndex = 0; TextureIndex < FilteredTextureElements.Num(); TextureIndex++ )
 			{
-				break;
-			}
+				ImportContext.bUserCancelled |= FDatasmithImporterImpl::HasUserCancelledTask( ImportContext.FeedbackContext );
 
-			// Skip asynchronous creation for IES textures as it is not supported.
-			// The creation of such textures will synchronously happen later in the call to FDatasmithTextureImporter::CreateTexture
-			if (FilteredTextureElements[TextureIndex]->GetTextureMode() == EDatasmithTextureMode::Ies)
-			{
-				continue;
-			}
+				if (ImportContext.bUserCancelled)
+				{
+					break;
+				}
 
-			UE::Interchange::FAssetImportResultRef FutureTexture = DatasmithTextureImporter.CreateTextureAsync( FilteredTextureElements[TextureIndex] );
+				// Skip asynchronous creation for IES textures as it is not supported.
+				// The creation of such textures will synchronously happen later in the call to FDatasmithTextureImporter::CreateTexture
+				if (FilteredTextureElements[TextureIndex]->GetTextureMode() == EDatasmithTextureMode::Ies)
+				{
+					continue;
+				}
 
-			if ( FutureTexture->IsValid() )
-			{
-				FutureTexture->OnDone(
-					[ TextureElement = FilteredTextureElements[TextureIndex].ToSharedRef(), &ImportContext ]( UE::Interchange::FAssetImportResult& AssetImportResults )
-					{
-						ImportMetaDataForObject( ImportContext, TextureElement, AssetImportResults.GetFirstAssetOfClass( UTexture::StaticClass() ) );
-					}
-				);
+				UE::Interchange::FAssetImportResultRef FutureTexture = DatasmithTextureImporter.CreateTextureAsync( FilteredTextureElements[TextureIndex] );
 
-				ImportContext.ImportedTextures.Add( FilteredTextureElements[TextureIndex].ToSharedRef(), MoveTemp( FutureTexture ) );
+				if ( FutureTexture->IsValid() )
+				{
+					FutureTexture->OnDone(
+						[ TextureElement = FilteredTextureElements[TextureIndex].ToSharedRef(), &ImportContext ]( UE::Interchange::FAssetImportResult& AssetImportResults )
+						{
+							ImportMetaDataForObject( ImportContext, TextureElement, AssetImportResults.GetFirstAssetOfClass( UTexture::StaticClass() ) );
+						}
+					);
+
+					ImportContext.ImportedTextures.Add( FilteredTextureElements[TextureIndex].ToSharedRef(), MoveTemp( FutureTexture ) );
+				}
 			}
 		}
 
