Fixed an issue where datasmith was asking interchange to import files that it coudn't.

diff --git a/Engine/Plugins/Enterprise/DatasmithImporter/Source/DatasmithImporter/Private/DatasmithTextureImporter.cpp b/Engine/Plugins/Enterprise/DatasmithImporter/Source/DatasmithImporter/Private/DatasmithTextureImporter.cpp
--- a/Engine/Plugins/Enterprise/DatasmithImporter/Source/DatasmithImporter/Private/DatasmithTextureImporter.cpp
+++ b/Engine/Plugins/Enterprise/DatasmithImporter/Source/DatasmithImporter/Private/DatasmithTextureImporter.cpp
@@ -383,6 +383,11 @@ UE::Interchange::FAssetImportResultRef FDatasmithTextureImporter::CreateTextureA
 {
 	UE::Interchange::FScopedSourceData ScopedSourceData( TextureElement->GetFile() );
 
+	if ( !UInterchangeManager::GetInterchangeManager().CanTranslateSourceData( ScopedSourceData.GetSourceData() ) )
+	{
+		return MakeShared< UE::Interchange::FAssetImportResult, ESPMode::ThreadSafe >();
+	}
+
 	const FString ContentPath = ImportContext.AssetsContext.TexturesImportPackage->GetPathName();
 
 	FImportAssetParameters ImportAssetParameters;
@@ -391,7 +396,7 @@ UE::Interchange::FAssetImportResultRef FDatasmithTextureImporter::CreateTextureA
 	UDatasmithTexturePipeline* TexturePipeline = NewObject< UDatasmithTexturePipeline >();
 	TexturePipeline->TextureElement = TextureElement;
 
-	ImportAssetParameters.OverridePipelines.Add(TexturePipeline);
+	ImportAssetParameters.OverridePipelines.Add( TexturePipeline );
 
 	return UInterchangeManager::GetInterchangeManager().ImportAssetAsync( ContentPath, ScopedSourceData.GetSourceData(), ImportAssetParameters );
 }
