Ensure HiResSourceModel has valid bulk data object after duplicating it from the source mesh which has had its source data removed.

diff --git a/Engine/Plugins/Enterprise/DatasmithImporter/Source/DatasmithImporter/Private/Utility/DatasmithImporterUtils.cpp b/Engine/Plugins/Enterprise/DatasmithImporter/Source/DatasmithImporter/Private/Utility/DatasmithImporterUtils.cpp
--- a/Engine/Plugins/Enterprise/DatasmithImporter/Source/DatasmithImporter/Private/Utility/DatasmithImporterUtils.cpp
+++ b/Engine/Plugins/Enterprise/DatasmithImporter/Source/DatasmithImporter/Private/Utility/DatasmithImporterUtils.cpp
@@ -1236,6 +1236,7 @@ UStaticMesh* FDatasmithImporterUtils::DuplicateStaticMesh(UStaticMesh* SourceSta
 	// Get rid of our temporary flag
 	SourceStaticMesh->ClearFlags(RF_ArchetypeObject);
 	DuplicateMesh->ClearFlags(RF_ArchetypeObject);
+	DuplicateMesh->GetHiResSourceModel().CreateSubObjects(DuplicateMesh);
 
 	if (bIgnoreBulkData)
 	{
