Datasmith: prevent ensure in automation test

- Displacement mapping is no longer accessible from datasmith (or at least, engine API changed). This CL prevents a crash when a dtasmith scene tries to use a displacement expression

diff --git a/Engine/Plugins/Enterprise/DatasmithImporter/Source/DatasmithImporter/Private/DatasmithMaterialExpressions.cpp b/Engine/Plugins/Enterprise/DatasmithImporter/Source/DatasmithImporter/Private/DatasmithMaterialExpressions.cpp
--- a/Engine/Plugins/Enterprise/DatasmithImporter/Source/DatasmithImporter/Private/DatasmithMaterialExpressions.cpp
+++ b/Engine/Plugins/Enterprise/DatasmithImporter/Source/DatasmithImporter/Private/DatasmithMaterialExpressions.cpp
@@ -1023,7 +1023,9 @@ void FDatasmithMaterialExpressions::ConnectToSlot(UMaterialExpression* ToBeConne
 
 	UMaterialFunction* Func = Cast< UMaterialFunction >( UnrealMatOrFunc );
 
-	if ((Slot == EDatasmithTextureSlot::METALLIC || Slot == EDatasmithTextureSlot::SPECULAR || Slot == EDatasmithTextureSlot::ROUGHNESS || Slot == EDatasmithTextureSlot::DISPLACE) && !ToBeConnected->IsA< UMaterialExpressionConstant >() && InputChannel == 0)
+	if ((Slot == EDatasmithTextureSlot::METALLIC || Slot == EDatasmithTextureSlot::SPECULAR || Slot == EDatasmithTextureSlot::ROUGHNESS)
+	  && !ToBeConnected->IsA< UMaterialExpressionConstant >()
+	  && InputChannel == 0)
 	{
 		if (UnrealMaterial != nullptr)
 		{
@@ -1047,8 +1049,6 @@ void FDatasmithMaterialExpressions::ConnectToSlot(UMaterialExpression* ToBeConne
 		case EDatasmithTextureSlot::OPACITYMASK:
 			UnrealMaterial->BlendMode = EBlendMode::BLEND_Masked;
 			break;
-		case EDatasmithTextureSlot::DISPLACE:
-			break;
 		}
 	}
 
@@ -2433,14 +2433,14 @@ void FDatasmithMaterialExpressions::CreateDatasmithMaterialHelper(UPackage* Pack
 
 	ModulateNormalAmount(UnrealMaterial, ShaderElement->GetBumpAmount());
 
-	// Displace mapping
-	if (AddCompExpression(ShaderElement->GetDisplaceComp(), UnrealMaterial, EDatasmithTextureSlot::DISPLACE, AssetsContext) == nullptr)
-	{
-		if (DisplaceTexture != nullptr)
-		{
-			AddTextureExpression(DisplaceTexture, ShaderElement->GetDisplaceTextureSampler(), 0.0, 0.0, UnrealMaterial, EDatasmithTextureSlot::DISPLACE);
-		}
-	}
+	// Displace mapping. #jira UE-113727
+// 	if (AddCompExpression(ShaderElement->GetDisplaceComp(), UnrealMaterial, EDatasmithTextureSlot::DISPLACE, AssetsContext) == nullptr)
+// 	{
+// 		if (DisplaceTexture != nullptr)
+// 		{
+// 			AddTextureExpression(DisplaceTexture, ShaderElement->GetDisplaceTextureSampler(), 0.0, 0.0, UnrealMaterial, EDatasmithTextureSlot::DISPLACE);
+// 		}
+// 	}
 }
 
 UMaterialExpressionMaterialFunctionCall* FDatasmithMaterialExpressions::BlendFunctions(UMaterial* UnrealMaterial, const FDatasmithAssetsImportContext& AssetsContext,
@@ -3162,7 +3162,7 @@ void FDatasmithMaterialExpressions::CreateUEPbrMaterialGraph(const TSharedPtr< I
 	}
 
 	ConnectExpression( MaterialElement.ToSharedRef(), MaterialExpressions, MaterialElement->GetNormal().GetExpression(), GetMaterialOrFunctionSlot( UnrealMaterialOrFunction, EDatasmithTextureSlot::NORMAL ), MaterialElement->GetNormal().GetOutputIndex() );
-	ConnectExpression( MaterialElement.ToSharedRef(), MaterialExpressions, MaterialElement->GetWorldDisplacement().GetExpression(), GetMaterialOrFunctionSlot( UnrealMaterialOrFunction, EDatasmithTextureSlot::DISPLACE ), MaterialElement->GetWorldDisplacement().GetOutputIndex() );
+// #jira UE-113727 // 	ConnectExpression( MaterialElement.ToSharedRef(), MaterialExpressions, MaterialElement->GetWorldDisplacement().GetExpression(), GetMaterialOrFunctionSlot( UnrealMaterialOrFunction, EDatasmithTextureSlot::DISPLACE ), MaterialElement->GetWorldDisplacement().GetOutputIndex() );
 
 	// Connect expressions to any UMaterialExpressionCustomOutput since these aren't part of the predefined material outputs
 	for ( int32 ExpressionIndex = 0; ExpressionIndex < MaterialElement->GetExpressionsCount(); ++ExpressionIndex )
diff --git a/Engine/Plugins/Enterprise/DatasmithImporter/Source/DatasmithImporter/Public/DatasmithMaterialExpressions.h b/Engine/Plugins/Enterprise/DatasmithImporter/Source/DatasmithImporter/Public/DatasmithMaterialExpressions.h
--- a/Engine/Plugins/Enterprise/DatasmithImporter/Source/DatasmithImporter/Public/DatasmithMaterialExpressions.h
+++ b/Engine/Plugins/Enterprise/DatasmithImporter/Source/DatasmithImporter/Public/DatasmithMaterialExpressions.h
@@ -47,8 +47,11 @@ enum class EDatasmithTextureSlot
 	OPACITYMASK,
 	NORMAL,
 	WORLDPOSITIONOFFSET,
-	DISPLACE,
-	TESSELLATIONMULTIPLIER,
+
+	// no longer meaningful. #jira UE-113727
+// 	DISPLACE,
+// 	TESSELLATIONMULTIPLIER,
+
 	SUBSURFACECOLOR,
 	COATSPECULAR,
 	COATROUGHNESS,
