datasmith: fix options not saved between sessions

#rnx

diff --git a/Engine/Plugins/Enterprise/DatasmithImporter/Source/DatasmithImporter/Private/DatasmithImportContext.cpp b/Engine/Plugins/Enterprise/DatasmithImporter/Source/DatasmithImporter/Private/DatasmithImportContext.cpp
--- a/Engine/Plugins/Enterprise/DatasmithImporter/Source/DatasmithImporter/Private/DatasmithImportContext.cpp
+++ b/Engine/Plugins/Enterprise/DatasmithImporter/Source/DatasmithImporter/Private/DatasmithImportContext.cpp
@@ -26,7 +26,7 @@ namespace DatasmithImportContextImpl
 {
 	const TCHAR* GetUserOptionPath()
 	{
-		static FString UserDatasmithOptionsFile = FPaths::Combine(FPlatformProcess::UserSettingsDir(), "Unreal Engine/Enterprise/Datasmith/Config/UserOptions.ini");
+		static FString UserDatasmithOptionsFile = FPaths::Combine(FPlatformProcess::UserSettingsDir(), TEXT("Unreal Engine/Enterprise/Datasmith/Config/UserOptions.ini"));
 		return *UserDatasmithOptionsFile;
 	}
 
