Datasmith: potential fix for crash in FStaticMeshOperations::ComputeTangentsAndNormals
#rnx

diff --git a/Engine/Plugins/Enterprise/DatasmithImporter/Source/DatasmithImporter/Private/DatasmithStaticMeshImporter.cpp b/Engine/Plugins/Enterprise/DatasmithImporter/Source/DatasmithImporter/Private/DatasmithStaticMeshImporter.cpp
--- a/Engine/Plugins/Enterprise/DatasmithImporter/Source/DatasmithImporter/Private/DatasmithStaticMeshImporter.cpp
+++ b/Engine/Plugins/Enterprise/DatasmithImporter/Source/DatasmithImporter/Private/DatasmithStaticMeshImporter.cpp
@@ -98,9 +98,9 @@ void FDatasmithStaticMeshImporter::CleanupMeshDescriptions(TArray<FMeshDescripti
 {
 	TRACE_CPUPROFILER_EVENT_SCOPE(FDatasmithStaticMeshImporter::CleanupMeshDescriptions)
 
-	TSet<FPolygonID> PolygonsToDelete;
 	for (FMeshDescription& MeshDescription : MeshDescriptions)
 	{
+		TSet<FPolygonID> PolygonsToDelete;
 		FStaticMeshAttributes Attributes(MeshDescription);
 		const TVertexAttributesRef<FVector3f> VertexPositions = Attributes.GetVertexPositions();
 		if (VertexPositions.IsValid())
@@ -118,6 +118,12 @@ void FDatasmithStaticMeshImporter::CleanupMeshDescriptions(TArray<FMeshDescripti
 			}
 		}
 
+		if (MeshDescription.Polygons().Num() == PolygonsToDelete.Num())
+		{
+			MeshDescription.Empty();
+			continue;
+		}
+
 		// Remove UV channels containing UV coordinates with Nan.
 		// This crashes the editor on Mac, see UE-106145.
 		TVertexInstanceAttributesRef<FVector2D> VertexInstanceUVs = Attributes.GetVertexInstanceUVs();
@@ -171,6 +177,16 @@ void FDatasmithStaticMeshImporter::CleanupMeshDescriptions(TArray<FMeshDescripti
 			MeshDescription.Compact(Remappings);
 		}
 
+		if (PolygonsToDelete.Num() > 0 || MeshDescription.Triangles().Num() == 0)
+		{
+			MeshDescription.TriangulateMesh();
+			if (MeshDescription.Triangles().Num() == 0)
+			{
+				MeshDescription.Empty();
+				continue;
+			}
+		}
+
 		// Fix invalid vertex normals and tangents
 		// We need polygon info because ComputeTangentsAndNormals uses it to repair the invalid vertex normals/tangents
 		// Can't calculate just the required polygons as ComputeTangentsAndNormals is parallel and we can't guarantee thread-safe access patterns
diff --git a/Engine/Source/Runtime/StaticMeshDescription/Private/StaticMeshOperations.cpp b/Engine/Source/Runtime/StaticMeshDescription/Private/StaticMeshOperations.cpp
--- a/Engine/Source/Runtime/StaticMeshDescription/Private/StaticMeshOperations.cpp
+++ b/Engine/Source/Runtime/StaticMeshDescription/Private/StaticMeshOperations.cpp
@@ -1321,6 +1321,11 @@ void FStaticMeshOperations::ComputeTangentsAndNormals(FMeshDescription& MeshDesc
 	}
 #endif
 
+	if (MeshDescription.Triangles().Num() == 0)
+	{
+		return;
+	}
+
 	const bool bForceComputeNormals = EnumHasAllFlags(ComputeNTBsOptions, EComputeNTBsFlags::Normals);
 	const bool bForceComputeTangent = EnumHasAnyFlags(ComputeNTBsOptions, EComputeNTBsFlags::Normals | EComputeNTBsFlags::Tangents);
 	const bool bComputeTangentWithMikkTSpace = bForceComputeTangent && EnumHasAllFlags(ComputeNTBsOptions, EComputeNTBsFlags::UseMikkTSpace);
