Fixed a crash when datasmith import a cubemap.
A pointer was badly cast and this resulted in datasmith writing invalid data into the data members of a UTextureCube.

diff --git a/Engine/Plugins/Enterprise/DatasmithImporter/Source/DatasmithImporter/Private/DatasmithTextureImporter.cpp b/Engine/Plugins/Enterprise/DatasmithImporter/Source/DatasmithImporter/Private/DatasmithTextureImporter.cpp
--- a/Engine/Plugins/Enterprise/DatasmithImporter/Source/DatasmithImporter/Private/DatasmithTextureImporter.cpp
+++ b/Engine/Plugins/Enterprise/DatasmithImporter/Source/DatasmithImporter/Private/DatasmithTextureImporter.cpp
@@ -263,9 +263,10 @@ UTexture* FDatasmithTextureImporter::CreateTexture(const TSharedPtr<IDatasmithTe
 
 	// This has to be called explicitly each time we create a texture since the flag gets reset in FactoryCreateBinary
 	TextureFact->SuppressImportOverwriteDialog();
-	UTexture2D* Texture =
-		(UTexture2D*)TextureFact->FactoryCreateBinary(UTexture2D::StaticClass(), TextureOuter, *TextureName, ImportContext.ObjectFlags /*& ~RF_Public*/, nullptr,
-															*Extension, PtrTextureData, PtrTextureDataEnd, GWarn);
+
+	UTexture* Texture =
+		Cast<UTexture>(TextureFact->FactoryCreateBinary(UTexture2D::StaticClass(), TextureOuter, *TextureName, ImportContext.ObjectFlags /*& ~RF_Public*/, nullptr,
+															*Extension, PtrTextureData, PtrTextureDataEnd, GWarn));
 	if (Texture != nullptr)
 	{
 		static_assert(TextureAddress::TA_Wrap == (int)EDatasmithTextureAddress::Wrap && TextureAddress::TA_Mirror == (int)EDatasmithTextureAddress::Mirror, "Texture Address enum doesn't match!" );
@@ -292,8 +293,12 @@ UTexture* FDatasmithTextureImporter::CreateTexture(const TSharedPtr<IDatasmithTe
 		}
 
 		Texture->Filter = TexFilter;
-		Texture->AddressX = (TextureAddress)TextureElement->GetTextureAddressX();
-		Texture->AddressY = (TextureAddress)TextureElement->GetTextureAddressY();
+
+		if (UTexture2D* Texture2D = Cast<UTexture2D>(Texture))
+		{
+			Texture2D->AddressX = (TextureAddress)TextureElement->GetTextureAddressX();
+			Texture2D->AddressY = (TextureAddress)TextureElement->GetTextureAddressY();
+		}
 
 		// Update import data
 		Texture->AssetImportData->Update(Filename, &Hash);
