Datasmith: static analysis pass

diff --git a/Engine/Plugins/Enterprise/DatasmithContent/Source/DatasmithContent/Private/ObjectTemplates/DatasmithAreaLightActorTemplate.cpp b/Engine/Plugins/Enterprise/DatasmithContent/Source/DatasmithContent/Private/ObjectTemplates/DatasmithAreaLightActorTemplate.cpp
--- a/Engine/Plugins/Enterprise/DatasmithContent/Source/DatasmithContent/Private/ObjectTemplates/DatasmithAreaLightActorTemplate.cpp
+++ b/Engine/Plugins/Enterprise/DatasmithContent/Source/DatasmithContent/Private/ObjectTemplates/DatasmithAreaLightActorTemplate.cpp
@@ -7,7 +7,7 @@
 UDatasmithAreaLightActorTemplate::UDatasmithAreaLightActorTemplate()
 	: UDatasmithObjectTemplate(true)
 {
-	Load( ADatasmithAreaLightActor::StaticClass()->GetDefaultObject() );
+	UDatasmithAreaLightActorTemplate::Load( ADatasmithAreaLightActor::StaticClass()->GetDefaultObject() );
 }
 
 UObject* UDatasmithAreaLightActorTemplate::UpdateObject( UObject* Destination, bool bForce )
diff --git a/Engine/Plugins/Enterprise/DatasmithContent/Source/DatasmithContent/Private/ObjectTemplates/DatasmithPointLightComponentTemplate.cpp b/Engine/Plugins/Enterprise/DatasmithContent/Source/DatasmithContent/Private/ObjectTemplates/DatasmithPointLightComponentTemplate.cpp
--- a/Engine/Plugins/Enterprise/DatasmithContent/Source/DatasmithContent/Private/ObjectTemplates/DatasmithPointLightComponentTemplate.cpp
+++ b/Engine/Plugins/Enterprise/DatasmithContent/Source/DatasmithContent/Private/ObjectTemplates/DatasmithPointLightComponentTemplate.cpp
@@ -6,7 +6,7 @@
 
 UDatasmithPointLightComponentTemplate::UDatasmithPointLightComponentTemplate()
 {
-	Load( UPointLightComponent::StaticClass()->GetDefaultObject() );
+	UDatasmithPointLightComponentTemplate::Load( UPointLightComponent::StaticClass()->GetDefaultObject() );
 }
 
 UObject* UDatasmithPointLightComponentTemplate::UpdateObject( UObject* Destination, bool bForce )
diff --git a/Engine/Plugins/Enterprise/DatasmithContent/Source/DatasmithContent/Private/ObjectTemplates/DatasmithSkyLightComponentTemplate.cpp b/Engine/Plugins/Enterprise/DatasmithContent/Source/DatasmithContent/Private/ObjectTemplates/DatasmithSkyLightComponentTemplate.cpp
--- a/Engine/Plugins/Enterprise/DatasmithContent/Source/DatasmithContent/Private/ObjectTemplates/DatasmithSkyLightComponentTemplate.cpp
+++ b/Engine/Plugins/Enterprise/DatasmithContent/Source/DatasmithContent/Private/ObjectTemplates/DatasmithSkyLightComponentTemplate.cpp
@@ -4,7 +4,7 @@
 
 UDatasmithSkyLightComponentTemplate::UDatasmithSkyLightComponentTemplate()
 {
-	Load( USkyLightComponent::StaticClass()->GetDefaultObject() );
+	UDatasmithSkyLightComponentTemplate::Load( USkyLightComponent::StaticClass()->GetDefaultObject() );
 }
 
 UObject* UDatasmithSkyLightComponentTemplate::UpdateObject( UObject* Destination, bool bForce )
diff --git a/Engine/Plugins/Enterprise/DatasmithContent/Source/DatasmithContent/Private/ObjectTemplates/DatasmithSpotLightComponentTemplate.cpp b/Engine/Plugins/Enterprise/DatasmithContent/Source/DatasmithContent/Private/ObjectTemplates/DatasmithSpotLightComponentTemplate.cpp
--- a/Engine/Plugins/Enterprise/DatasmithContent/Source/DatasmithContent/Private/ObjectTemplates/DatasmithSpotLightComponentTemplate.cpp
+++ b/Engine/Plugins/Enterprise/DatasmithContent/Source/DatasmithContent/Private/ObjectTemplates/DatasmithSpotLightComponentTemplate.cpp
@@ -6,7 +6,7 @@
 
 UDatasmithSpotLightComponentTemplate::UDatasmithSpotLightComponentTemplate()
 {
-	Load( USpotLightComponent::StaticClass()->GetDefaultObject() );
+	UDatasmithSpotLightComponentTemplate::Load( USpotLightComponent::StaticClass()->GetDefaultObject() );
 }
 
 UObject* UDatasmithSpotLightComponentTemplate::UpdateObject( UObject* Destination, bool bForce )
diff --git a/Engine/Plugins/Enterprise/DatasmithGLTFImporter/Source/DatasmithGLTFTranslator/Private/DatasmithGLTFTextureFactory.cpp b/Engine/Plugins/Enterprise/DatasmithGLTFImporter/Source/DatasmithGLTFTranslator/Private/DatasmithGLTFTextureFactory.cpp
--- a/Engine/Plugins/Enterprise/DatasmithGLTFImporter/Source/DatasmithGLTFTranslator/Private/DatasmithGLTFTextureFactory.cpp
+++ b/Engine/Plugins/Enterprise/DatasmithGLTFImporter/Source/DatasmithGLTFTranslator/Private/DatasmithGLTFTextureFactory.cpp
@@ -78,13 +78,13 @@ namespace DatasmithGLTFImporterImpl
 	}
 }
 FDatasmithGLTFTextureFactory::FDatasmithGLTFTextureFactory()
-    : CurrentScene(nullptr)
+	: CurrentScene(nullptr)
 {
 }
 
 FDatasmithGLTFTextureFactory::~FDatasmithGLTFTextureFactory()
 {
-	CleanUp();
+	FDatasmithGLTFTextureFactory::CleanUp();
 }
 
 GLTF::ITextureElement* FDatasmithGLTFTextureFactory::CreateTexture(const GLTF::FTexture& GltfTexture, UObject* ParentPackage, EObjectFlags Flags,
diff --git a/Engine/Plugins/Enterprise/DatasmithImporter/Source/DatasmithImporter/Private/UVTools/UVGenerationFlattenMapping.cpp b/Engine/Plugins/Enterprise/DatasmithImporter/Source/DatasmithImporter/Private/UVTools/UVGenerationFlattenMapping.cpp
--- a/Engine/Plugins/Enterprise/DatasmithImporter/Source/DatasmithImporter/Private/UVTools/UVGenerationFlattenMapping.cpp
+++ b/Engine/Plugins/Enterprise/DatasmithImporter/Source/DatasmithImporter/Private/UVTools/UVGenerationFlattenMapping.cpp
@@ -43,7 +43,7 @@ public:
 		int32 Group;
 		int32 PolygonID;
 	};
-	
+
 	static TArray<FaceStruct> GetFacesFrom(FMeshDescription& InMesh, const TArray<int32>& OverlappingCorners={});
 	static void CalculateUVs(TArray<FaceStruct>& FlattenFaces, int32 VertexNum, float AngleThreshold, float AreaWeight);
 
@@ -695,7 +695,7 @@ void FUVGenerationFlattenMappingInternal::CreateGroupUVIslands(TSet<FaceStruct*>
 			}
 		}
 
-		if (MostDifferentAngle < SlopeProjectionLimit)
+		if (MostUniqueFace && MostDifferentAngle < SlopeProjectionLimit)
 		{
 			NewProjectVec = MostUniqueFace->Normal;
 			NextProjectionFaces.Empty();
@@ -829,7 +829,7 @@ void FUVGenerationFlattenMappingInternal::CreateGroupUVIslands(TSet<FaceStruct*>
 				TMap<int32, int32>& EdgeProjectionGroup = EdgeProjectionGroupList[PrefAngleIdx];
 
 				int32 UnconnectedVert = 0;
-				
+
 				const int32* Value = EdgeProjectionGroup.Find(Face->V0);
 				if (Value == nullptr || *Value < 2)
 				{
@@ -1190,14 +1190,14 @@ void UUVGenerationFlattenMapping::GenerateFlattenMappingUVs(UStaticMesh* InStati
 	{
 		return;
 	}
-	
+
 	InStaticMesh->Modify();
 	FScopedSlowTask SlowTask(InStaticMesh->GetNumSourceModels(), FText::FromString(InStaticMesh->GetName()));
 
 	for (int32 LodIndex = 0; LodIndex < InStaticMesh->GetNumSourceModels(); ++LodIndex)
 	{
 		SlowTask.EnterProgressFrame(1);
-		
+
 		if (!InStaticMesh->IsMeshDescriptionValid(LodIndex))
 		{
 			continue;
@@ -1233,7 +1233,7 @@ void UUVGenerationFlattenMapping::GenerateUVs(FMeshDescription& InMesh, int32 UV
 
 	// Get the Internal Face Array
 	TArray< FUVGenerationFlattenMappingInternal::FaceStruct > FlattenFaces = FUVGenerationFlattenMappingInternal::GetFacesFrom(InMesh, OverlappingCorners);
-	
+
 	FUVGenerationFlattenMappingInternal::CalculateUVs(FlattenFaces, InMesh.Vertices().Num(), AngleThreshold, FLATTEN_AREA_WEIGHT);
 
 	// Write unwrapped UVs back to mesh description
