Github PR#8044: Added clear cloat and subsurface shading models to datasmith importer

#rnx

diff --git a/Engine/Plugins/Enterprise/DatasmithImporter/Source/DatasmithImporter/Private/DatasmithMaterialExpressions.cpp b/Engine/Plugins/Enterprise/DatasmithImporter/Source/DatasmithImporter/Private/DatasmithMaterialExpressions.cpp
--- a/Engine/Plugins/Enterprise/DatasmithImporter/Source/DatasmithImporter/Private/DatasmithMaterialExpressions.cpp
+++ b/Engine/Plugins/Enterprise/DatasmithImporter/Source/DatasmithImporter/Private/DatasmithMaterialExpressions.cpp
@@ -3271,11 +3271,30 @@ UMaterialInterface* FDatasmithMaterialExpressions::CreateUEPbrMaterial(UPackage*
 
 	CreateUEPbrMaterialGraph(MaterialElement, AssetsContext, UnrealMaterial);
 
+	EDatasmithShadingModel MaterialShadingModel = MaterialElement->GetShadingModel();
+	switch (MaterialShadingModel)
+	{
+		case EDatasmithShadingModel::Subsurface: 
+			UnrealMaterial->SetShadingModel(MSM_Subsurface); 
+			break;
+
+		case EDatasmithShadingModel::ClearCoat: 
+			UnrealMaterial->SetShadingModel(MSM_ClearCoat); 
+			break;
+
+		case EDatasmithShadingModel::ThinTranslucent:
+			UnrealMaterial->SetShadingModel(MSM_ThinTranslucent);
+			break;
+
+		default:
+			break;
+	}
+	
+
 	if ( MaterialElement->GetOpacity().GetExpression() )
 	{
 		if ( MaterialElement->GetShadingModel() == EDatasmithShadingModel::ThinTranslucent )
 		{
-			UnrealMaterial->SetShadingModel( MSM_ThinTranslucent );
 			UnrealMaterial->TranslucencyLightingMode = TLM_SurfacePerPixelLighting;
 		}
 		else
diff --git a/Engine/Source/Runtime/Datasmith/DatasmithCore/Public/DatasmithDefinitions.h b/Engine/Source/Runtime/Datasmith/DatasmithCore/Public/DatasmithDefinitions.h
--- a/Engine/Source/Runtime/Datasmith/DatasmithCore/Public/DatasmithDefinitions.h
+++ b/Engine/Source/Runtime/Datasmith/DatasmithCore/Public/DatasmithDefinitions.h
@@ -340,12 +340,14 @@ enum class EDatasmithShaderUsage
 	LightFunction
 };
 
-static const TCHAR* DatasmithShadingModelStrings[] = { TEXT("DefaultLit"), TEXT("ThinTranslucent") };
+static const TCHAR* DatasmithShadingModelStrings[] = { TEXT("DefaultLit"), TEXT("ThinTranslucent"), TEXT("Subsurface"), TEXT("ClearCoat") };
 
 enum class EDatasmithShadingModel : uint8
 {
 	DefaultLit,
 	ThinTranslucent,
+	Subsurface,
+	ClearCoat
 };
 
 UENUM()
