Visual dataprep map asset is removed after deleting levels in an associated level

diff --git a/Engine/Plugins/Enterprise/DatasmithImporter/Source/DatasmithImporter/Private/DatasmithConsumer.cpp b/Engine/Plugins/Enterprise/DatasmithImporter/Source/DatasmithImporter/Private/DatasmithConsumer.cpp
--- a/Engine/Plugins/Enterprise/DatasmithImporter/Source/DatasmithImporter/Private/DatasmithConsumer.cpp
+++ b/Engine/Plugins/Enterprise/DatasmithImporter/Source/DatasmithImporter/Private/DatasmithConsumer.cpp
@@ -979,6 +979,25 @@ ULevel* UDatasmithConsumer::FindOrAddLevel(const FString& InLevelName)
 	{
 		StreamingLevel = EditorLevelUtils::CreateNewStreamingLevelForWorld( *WorkingWorld, ULevelStreamingAlwaysLoaded::StaticClass(), *PackageFilename );
 		ensure(StreamingLevel);
+
+		// Register the newly created map asset (associated with this consumer) to the asset registry
+		UPackage* WorldPackage = FindPackage(nullptr, *StreamingLevel->GetWorldAssetPackageName());
+		ensure(WorldPackage);
+
+		UWorld* World = nullptr;
+		ForEachObjectWithPackage(WorldPackage, [&World](UObject* Object)
+		{
+			if (UWorld* WorldAsset = Cast<UWorld>(Object))
+			{
+				World = WorldAsset;
+				return false;
+			}
+			return true;
+		}, false);
+
+		ensure(World);
+		
+		FAssetRegistryModule::AssetCreated(World);
 	}
 
 	WorkingWorld->PersistentLevel = CurrentLevel;
