DirectLink in Editor - Detecting and cancelling loading of empty DirectLink sources.

JeanLuc.Corenthin

diff --git a/Engine/Plugins/Enterprise/DatasmithImporter/Source/DirectLinkExtension/Private/DirectLinkManager.cpp b/Engine/Plugins/Enterprise/DatasmithImporter/Source/DirectLinkExtension/Private/DirectLinkManager.cpp
--- a/Engine/Plugins/Enterprise/DatasmithImporter/Source/DirectLinkExtension/Private/DirectLinkManager.cpp
+++ b/Engine/Plugins/Enterprise/DatasmithImporter/Source/DirectLinkExtension/Private/DirectLinkManager.cpp
@@ -174,6 +174,8 @@ namespace UE::DatasmithImporter
 		}
 
 		UpdateSourceCache();
+
+		CancelEmptySourcesLoading();
 	}
 
 	void FDirectLinkManager::UpdateSourceCache()
@@ -211,6 +213,34 @@ namespace UE::DatasmithImporter
 		}
 	}
 
+	void FDirectLinkManager::CancelEmptySourcesLoading() const
+	{
+		FRWScopeLock ScopeLock(RawInfoLock, FRWScopeLockType::SLT_ReadOnly);
+
+		for (const DirectLink::FRawInfo::FStreamInfo& StreamInfo : RawInfoCache.StreamsInfo)
+		{
+			if (const TSharedRef<FDirectLinkExternalSource>* ExternalSourcePtr = DirectLinkSourceToExternalSourceMap.Find(StreamInfo.Source))
+			{
+				const TSharedRef<FDirectLinkExternalSource>& ExternalSource = *ExternalSourcePtr;
+				
+				// We can infer that a DirectLink source is empty (no scene synced) by looking at if its stream is planning to send any data.
+				// #ueent_todo: Ideally it would be better to not allow an AsyncLoad to take place in the first time, but we can't know a source is empty 
+				//				before actually connecting to it, so this is the best we can do at the current time.
+				const bool bStreamIsEmpty = StreamInfo.bIsActive
+					&& !StreamInfo.CommunicationStatus.IsTransmitting()
+					&& StreamInfo.CommunicationStatus.TaskTotal == 0;
+
+				if (bStreamIsEmpty
+					&& ExternalSource->IsAsyncLoading()
+					&& !ExternalSource->GetDatasmithScene().IsValid())
+				{
+					ExternalSource->CancelAsyncLoad();
+					UE_LOG(LogDirectLinkManager, Error, TEXT("The DirectLink source \"%s\" could not be loaded: Nothing to synchronize. Make sure to do a DirectLink sync in your exporter."), *ExternalSource->GetSourceName());
+				}
+			}
+		}
+	}
+
 	void FDirectLinkManager::RegisterDirectLinkExternalSource(FDirectLinkExternalSourceRegisterInformation&& RegisterInformation)
 	{
 		RegisteredExternalSourcesInfo.Add(MoveTemp(RegisterInformation));
diff --git a/Engine/Plugins/Enterprise/DatasmithImporter/Source/DirectLinkExtension/Private/DirectLinkManager.h b/Engine/Plugins/Enterprise/DatasmithImporter/Source/DirectLinkExtension/Private/DirectLinkManager.h
--- a/Engine/Plugins/Enterprise/DatasmithImporter/Source/DirectLinkExtension/Private/DirectLinkManager.h
+++ b/Engine/Plugins/Enterprise/DatasmithImporter/Source/DirectLinkExtension/Private/DirectLinkManager.h
@@ -75,6 +75,8 @@ namespace UE::DatasmithImporter
 		 */
 		void UpdateSourceCache();
 
+		void CancelEmptySourcesLoading() const;
+
 		bool EnableAssetAutoReimport(UObject* InAsset);
 
 		bool DisableAssetAutoReimport(UObject* InAsset);
diff --git a/Engine/Plugins/Enterprise/DatasmithImporter/Source/DirectLinkExtension/Private/UI/SDirectLinkAssetThumbnailIndicator.cpp b/Engine/Plugins/Enterprise/DatasmithImporter/Source/DirectLinkExtension/Private/UI/SDirectLinkAssetThumbnailIndicator.cpp
--- a/Engine/Plugins/Enterprise/DatasmithImporter/Source/DirectLinkExtension/Private/UI/SDirectLinkAssetThumbnailIndicator.cpp
+++ b/Engine/Plugins/Enterprise/DatasmithImporter/Source/DirectLinkExtension/Private/UI/SDirectLinkAssetThumbnailIndicator.cpp
@@ -83,7 +83,7 @@ namespace UE::DatasmithImporter
 
 			const FMD5Hash SourceHash = TryGetSourceHash(AssetData);
 			const FMD5Hash ExternalSourceHash = ExternalSource->GetSourceHash();
-			if (SourceHash.IsValid() && SourceHash != ExternalSourceHash)
+			if (!SourceHash.IsValid() || SourceHash != ExternalSourceHash)
 			{
 				//The source hash in the asset is different from the one in the external source, the asset is out of sync.
 				return EDirectLinkSourceStatus::OutOfSync;
diff --git a/Engine/Plugins/Enterprise/DatasmithImporter/Source/ExternalSource/Private/ExternalSource.cpp b/Engine/Plugins/Enterprise/DatasmithImporter/Source/ExternalSource/Private/ExternalSource.cpp
--- a/Engine/Plugins/Enterprise/DatasmithImporter/Source/ExternalSource/Private/ExternalSource.cpp
+++ b/Engine/Plugins/Enterprise/DatasmithImporter/Source/ExternalSource/Private/ExternalSource.cpp
@@ -176,6 +176,11 @@ namespace UE::DatasmithImporter
 	{
 		OnExternalSourceChanged.Clear();
 
+		CancelAsyncLoad();
+	}
+
+	void FExternalSource::CancelAsyncLoad()
+	{
 		// Set value of all pending TPromise to null, to avoid creating deadlocks.
 		{
 			TSharedPtr<IDatasmithScene> NullScene;
diff --git a/Engine/Plugins/Enterprise/DatasmithImporter/Source/ExternalSource/Public/ExternalSource.h b/Engine/Plugins/Enterprise/DatasmithImporter/Source/ExternalSource/Public/ExternalSource.h
--- a/Engine/Plugins/Enterprise/DatasmithImporter/Source/ExternalSource/Public/ExternalSource.h
+++ b/Engine/Plugins/Enterprise/DatasmithImporter/Source/ExternalSource/Public/ExternalSource.h
@@ -122,6 +122,16 @@ namespace UE::DatasmithImporter
 		 */
 		TFuture<TSharedPtr<IDatasmithScene>> AsyncLoad();
 
+		/**
+		 * True when the ExternalSource is currently loading the source asynchronously.
+		 */
+		bool IsAsyncLoading() const { return !PendingPromiseQueue.IsEmpty(); }
+
+		/**
+		 * Cancel any pending AsyncLoad() operation, pending TFuture will return invalid TSharedPtr.
+		 */
+		void CancelAsyncLoad();
+
 		/**
 		 * Delegate called on main thread every time the loaded data is updated.
 		 * Used for registering auto-reimport on assets.
