DirectLink in Editor - Auto-reimport now ignores intermediate reimport when multiple DirectLink update are triggered while we are already importing. Reimporting now clears all pending reimport request.

JeanLuc.Corenthin

diff --git a/Engine/Plugins/Enterprise/DatasmithImporter/Source/DirectLinkExtension/Private/DirectLinkManager.cpp b/Engine/Plugins/Enterprise/DatasmithImporter/Source/DirectLinkExtension/Private/DirectLinkManager.cpp
--- a/Engine/Plugins/Enterprise/DatasmithImporter/Source/DirectLinkExtension/Private/DirectLinkManager.cpp
+++ b/Engine/Plugins/Enterprise/DatasmithImporter/Source/DirectLinkExtension/Private/DirectLinkManager.cpp
@@ -394,33 +394,53 @@ namespace UE::DatasmithImporter
 
 	void FDirectLinkManager::OnExternalSourceChanged(const TSharedRef<FExternalSource>& ExternalSource)
 	{
-		Async(EAsyncExecution::TaskGraphMainThread, [this, ExternalSource]() {
+		// Accumulate the reimport request in a thread-safe queue that will be processed in the main thread.
+		// Multiple reimport request for the same external source will only be processed once. 
+		// Doing this allows us to skip redundant reimports, as the reimport already uses the latest data from the ExternalSource.
+		PendingReimportQueue.Enqueue(ExternalSource);
 
-			TArray<TSharedRef<FAutoReimportInfo>> AutoReimportInfos;
-			RegisteredAutoReimportExternalSourceMap.MultiFind(ExternalSource, AutoReimportInfos);
-			if (AutoReimportInfos.Num() == 0)
+		Async(EAsyncExecution::TaskGraphMainThread, [this]() {
+
+			TSet<TSharedRef<FExternalSource>> PendingReimportSet;
+			TSharedPtr<FExternalSource> EnqueuedSourceToReimport;
+			while (PendingReimportQueue.Dequeue(EnqueuedSourceToReimport))
 			{
-				return;
+				PendingReimportSet.Add(EnqueuedSourceToReimport.ToSharedRef());
 			}
 
-			for (const TSharedRef<FAutoReimportInfo>& AutoReimportInfo : AutoReimportInfos)
+			for (const TSharedRef<FExternalSource>& ExternalSourceToReimport : PendingReimportSet)
 			{
+				TriggerAutoReimportOnExternalSource(ExternalSourceToReimport);
+			}
+		});
+	}
+
+	void FDirectLinkManager::TriggerAutoReimportOnExternalSource(const TSharedRef<FExternalSource>& ExternalSource)
+	{
+		TArray<TSharedRef<FAutoReimportInfo>> AutoReimportInfos;
+		RegisteredAutoReimportExternalSourceMap.MultiFind(ExternalSource, AutoReimportInfos);
+		if (AutoReimportInfos.Num() == 0)
+		{
+			return;
+		}
+
+		for (const TSharedRef<FAutoReimportInfo>& AutoReimportInfo : AutoReimportInfos)
+		{
 #if WITH_EDITOR
-				// If we're in PIE, delay the callbacks until we exit that mode.
-				if (GIsEditor && FApp::IsGame())
-				{
-					AutoReimportInfo->bChangedDuringPIE = true;
-					UE_LOG(LogDirectLinkManager, Warning, TEXT("The DirectLink source \"%s\" received an update while in PIE mode. The reimport will be triggered when exiting PIE."), *ExternalSource->GetSourceName());
-					continue;
-				}
+			// If we're in PIE, delay the callbacks until we exit that mode.
+			if (GIsEditor && FApp::IsGame())
+			{
+				AutoReimportInfo->bChangedDuringPIE = true;
+				UE_LOG(LogDirectLinkManager, Warning, TEXT("The DirectLink source \"%s\" received an update while in PIE mode. The reimport will be triggered when exiting PIE."), *ExternalSource->GetSourceName());
+				continue;
+			}
 #endif //WITH_EDITOR
 
-				if (UObject* Asset = AutoReimportInfo->TargetObject.Get())
-				{
-					TriggerAutoReimportOnAsset(AutoReimportInfo->TargetObject.Get());
-				}
+			if (UObject* Asset = AutoReimportInfo->TargetObject.Get())
+			{
+				TriggerAutoReimportOnAsset(AutoReimportInfo->TargetObject.Get());
 			}
-		});
+		}
 	}
 
 	void FDirectLinkManager::TriggerAutoReimportOnAsset(UObject* Asset)
diff --git a/Engine/Plugins/Enterprise/DatasmithImporter/Source/DirectLinkExtension/Private/DirectLinkManager.h b/Engine/Plugins/Enterprise/DatasmithImporter/Source/DirectLinkExtension/Private/DirectLinkManager.h
--- a/Engine/Plugins/Enterprise/DatasmithImporter/Source/DirectLinkExtension/Private/DirectLinkManager.h
+++ b/Engine/Plugins/Enterprise/DatasmithImporter/Source/DirectLinkExtension/Private/DirectLinkManager.h
@@ -83,6 +83,8 @@ namespace UE::DatasmithImporter
 
 		void OnExternalSourceChanged(const TSharedRef<FExternalSource>& ExternalSource);
 
+		void TriggerAutoReimportOnExternalSource(const TSharedRef<FExternalSource>& ExternalSource);
+
 		void TriggerAutoReimportOnAsset(UObject* Asset);
 
 #if WITH_EDITOR
@@ -114,6 +116,8 @@ namespace UE::DatasmithImporter
 		
 		TMultiMap<TSharedRef<FExternalSource>, TSharedRef<FAutoReimportInfo>> RegisteredAutoReimportExternalSourceMap;
 
+		TQueue<TSharedPtr<FExternalSource>> PendingReimportQueue;
+
 #if WITH_EDITOR
 		FDelegateHandle OnPIEEndHandle;
 #endif
