Datasmith Import - Fixing an issue that would add all the Datasmith-supported file formats to the standard import pipeline after reimporting a DatasmithScene.

JeanLuc.Corenthin

diff --git a/Engine/Plugins/Enterprise/DatasmithImporter/Source/DatasmithImporter/Private/DatasmithImportFactory.cpp b/Engine/Plugins/Enterprise/DatasmithImporter/Source/DatasmithImporter/Private/DatasmithImportFactory.cpp
--- a/Engine/Plugins/Enterprise/DatasmithImporter/Source/DatasmithImporter/Private/DatasmithImportFactory.cpp
+++ b/Engine/Plugins/Enterprise/DatasmithImporter/Source/DatasmithImporter/Private/DatasmithImportFactory.cpp
@@ -578,11 +578,9 @@ void UDatasmithImportFactory::ParseFromJson(TSharedRef<FJsonObject> InImportSett
 bool UDatasmithImportFactory::CanReimport( UObject* Obj, TArray<FString>& OutFilenames )
 {
 	// The CDO may be used to do that check. In that case, Formats are not necessarily initialized.
-	if (Formats.Num() == 0)
-	{
-		Formats = FDatasmithTranslatorManager::Get().GetSupportedFormats();
-	}
-
+	const bool bGenerateSupportedFormats = HasAnyFlags(RF_ClassDefaultObject) && Formats.Num() == 0;
+	TGuardValue<TArray<FString>> FormatsGuard(Formats, bGenerateSupportedFormats ? FDatasmithTranslatorManager::Get().GetSupportedFormats() : Formats);
+	
 	UAssetImportData* ReImportData = DatasmithImportFactoryImpl::GetImportData(Obj);
 	if (ReImportData == nullptr)
 	{
