DirectLink in Editor - Fixing unable to instanciate DatasmithScene asset when DirectLink source is unavailable.

JeanLuc.Corenthin

diff --git a/Engine/Plugins/Enterprise/DatasmithImporter/Source/DatasmithImporter/Private/ActorFactoryDatasmithScene.cpp b/Engine/Plugins/Enterprise/DatasmithImporter/Source/DatasmithImporter/Private/ActorFactoryDatasmithScene.cpp
--- a/Engine/Plugins/Enterprise/DatasmithImporter/Source/DatasmithImporter/Private/ActorFactoryDatasmithScene.cpp
+++ b/Engine/Plugins/Enterprise/DatasmithImporter/Source/DatasmithImporter/Private/ActorFactoryDatasmithScene.cpp
@@ -62,13 +62,8 @@ namespace UActorFactoryDatasmithSceneImpl
 			ExternalSource = IExternalSourceModule::Get().GetManager().TryGetExternalSourceFromImportData(*DatasmithImportData.Get());
 		}
 		
-		if (!ExternalSource)
-		{
-			return nullptr;
-		}
-
 		const FString ImportPath = DatasmithScene->AssetImportData->BaseOptions.AssetOptions.PackagePath.ToString();
-		FDatasmithImportContext ImportContext( ExternalSource.ToSharedRef(), false, LoggerName, LoggerLabel );
+		FDatasmithImportContext ImportContext( ExternalSource, false, LoggerName, LoggerLabel );
 
 		ImportContext.Options->BaseOptions = DatasmithScene->AssetImportData->BaseOptions;
 		ImportContext.Options->BaseOptions.SceneHandling = EDatasmithImportScene::CurrentLevel;
diff --git a/Engine/Plugins/Enterprise/DatasmithImporter/Source/DatasmithImporter/Private/DatasmithImportContext.cpp b/Engine/Plugins/Enterprise/DatasmithImporter/Source/DatasmithImporter/Private/DatasmithImportContext.cpp
--- a/Engine/Plugins/Enterprise/DatasmithImporter/Source/DatasmithImporter/Private/DatasmithImportContext.cpp
+++ b/Engine/Plugins/Enterprise/DatasmithImporter/Source/DatasmithImporter/Private/DatasmithImportContext.cpp
@@ -89,6 +89,8 @@ bool FDatasmithImportOptionHelper::DisplayOptionsDialog(const TArray<UObject*>&
 
 	// First option object is always a UDatasmithImportOptions object
 	UDatasmithImportOptions* MainOptions = Cast<UDatasmithImportOptions>(ImportOptions[0]);
+	ensure(MainOptions->FileName.Len() > 0);
+	ensure(MainOptions->FilePath.Len() > 0);
 
 	float SceneVersion;
 	FString FileSDKVersion;
@@ -208,9 +210,8 @@ FDatasmithImportContext::FDatasmithImportContext(const FString& InFileName, bool
 }
 
 
-FDatasmithImportContext::FDatasmithImportContext(const TSharedRef<UE::DatasmithImporter::FExternalSource>& InExternalSource, bool bLoadConfig, const FName& LoggerName, const FText& LoggerLabel)
-	: SceneTranslator(InExternalSource->GetAssetTranslator())
-	, Options(NewObject<UDatasmithImportOptions>(GetTransientPackage(), TEXT("Datasmith Import Settings")))
+FDatasmithImportContext::FDatasmithImportContext(const TSharedPtr<UE::DatasmithImporter::FExternalSource>& InExternalSource, bool bLoadConfig, const FName& LoggerName, const FText& LoggerLabel)
+	: Options(NewObject<UDatasmithImportOptions>(GetTransientPackage(), TEXT("Datasmith Import Settings")))
 	, SceneAsset(nullptr)
 	, bUserCancelled(false)
 	, bIsAReimport(false)
@@ -222,9 +223,22 @@ FDatasmithImportContext::FDatasmithImportContext(const TSharedRef<UE::DatasmithI
 	, CurrentSceneActorIndex(0)
 	, ReferenceCollector(this)
 {
-	if (SceneTranslator)
+	if (InExternalSource)
 	{
-		SceneTranslator->GetSceneImportOptions(AdditionalImportOptions);
+		SceneTranslator = InExternalSource->GetAssetTranslator();
+		if (SceneTranslator)
+		{
+			SceneTranslator->GetSceneImportOptions(AdditionalImportOptions);
+		}
+		
+		const FString FileName = InExternalSource->GetFallbackFilepath();
+		Options->FileName = FPaths::GetCleanFilename(FileName);
+		Options->FilePath = FPaths::ConvertRelativePathToFull(FileName);
+		Options->SourceUri = InExternalSource->GetSourceUri().ToString();
+		Options->SourceHash = InExternalSource->GetSourceHash();
+		SceneName = FDatasmithUtils::SanitizeObjectName(InExternalSource->GetSceneName());
+
+		FileHash = FMD5Hash::HashFile(*Options->FilePath);
 	}
 
 	if (bLoadConfig)
@@ -238,15 +252,7 @@ FDatasmithImportContext::FDatasmithImportContext(const TSharedRef<UE::DatasmithI
 
 	// Force the SceneHandling to be on current level by default.
 	// Note: This is done because this option was previously persisted and can get overwritten
-	const FString FileName = InExternalSource->GetFallbackFilepath();
 	Options->BaseOptions.SceneHandling = EDatasmithImportScene::CurrentLevel;
-	Options->FileName = FPaths::GetCleanFilename(FileName);
-	Options->FilePath = FPaths::ConvertRelativePathToFull(FileName);
-	Options->SourceUri = InExternalSource->GetSourceUri().ToString();
-	Options->SourceHash = InExternalSource->GetSourceHash();
-	SceneName = FDatasmithUtils::SanitizeObjectName(InExternalSource->GetSceneName());
-
-	FileHash = FMD5Hash::HashFile(*Options->FilePath);
 }
 
 void FDatasmithImportContext::UpdateImportOption(UDatasmithOptionsBase* NewOption)
@@ -310,9 +316,6 @@ bool FDatasmithImportContext::InitOptions(const TSharedPtr<FJsonObject>& ImportS
 		return false;
 	}
 
-	check(Options->FileName.Len() > 0);
-	check(Options->FilePath.Len() > 0);
-
 	TArray<UObject*> ImportOptions;
 	ImportOptions.Reserve(1 + AdditionalImportOptions.Num());
 
diff --git a/Engine/Plugins/Enterprise/DatasmithImporter/Source/DatasmithImporter/Public/DatasmithImportContext.h b/Engine/Plugins/Enterprise/DatasmithImporter/Source/DatasmithImporter/Public/DatasmithImportContext.h
--- a/Engine/Plugins/Enterprise/DatasmithImporter/Source/DatasmithImporter/Public/DatasmithImportContext.h
+++ b/Engine/Plugins/Enterprise/DatasmithImporter/Source/DatasmithImporter/Public/DatasmithImportContext.h
@@ -211,7 +211,7 @@ struct DATASMITHIMPORTER_API FDatasmithImportContext
 	UE_DEPRECATED(5.0, "Please use constructor using ExternalSource")
 	FDatasmithImportContext(const FString& InFileName, bool bLoadConfig, const FName& LoggerName, const FText& LoggerLabel, TSharedPtr<IDatasmithTranslator> InSceneTranslator = nullptr);
 
-	FDatasmithImportContext(const TSharedRef<UE::DatasmithImporter::FExternalSource>& InExternalSource, bool bLoadConfig, const FName& LoggerName, const FText& LoggerLabel);
+	FDatasmithImportContext(const TSharedPtr<UE::DatasmithImporter::FExternalSource>& InExternalSource, bool bLoadConfig, const FName& LoggerName, const FText& LoggerLabel);
 	
 	/** Cached MD5 hash value for faster processing */
 	FMD5Hash FileHash;
