Datasmith: Fix crash on import after a mesh is removed from the scene after translation and before import.

diff --git a/Engine/Plugins/Enterprise/DatasmithImporter/Source/DatasmithImporter/Private/DatasmithImporter.cpp b/Engine/Plugins/Enterprise/DatasmithImporter/Source/DatasmithImporter/Private/DatasmithImporter.cpp
--- a/Engine/Plugins/Enterprise/DatasmithImporter/Source/DatasmithImporter/Private/DatasmithImporter.cpp
+++ b/Engine/Plugins/Enterprise/DatasmithImporter/Source/DatasmithImporter/Private/DatasmithImporter.cpp
@@ -1590,10 +1590,7 @@ void FDatasmithImporter::ImportMetaDataForObject(FDatasmithImportContext& Import
 
 void FDatasmithImporter::FilterElementsToImport( FDatasmithImportContext& ImportContext )
 {
-	if (!ImportContext.FilteredScene.IsValid())
-	{
-		ImportContext.FilteredScene = FDatasmithSceneFactory::DuplicateScene( ImportContext.Scene.ToSharedRef() );
-	}
+	ImportContext.FilteredScene = FDatasmithSceneFactory::DuplicateScene( ImportContext.Scene.ToSharedRef() );
 
 	FDatasmithSceneUtils::CleanUpScene(ImportContext.FilteredScene.ToSharedRef(), false);
 
