DirectLink in Editor - Fix loading popup not always being canceled because of bad condition in IsAsyncLoading() causing a race condition + downgrade empty source error to warning

diff --git a/Engine/Plugins/Enterprise/DatasmithImporter/Source/DirectLinkExtension/Private/DirectLinkManager.cpp b/Engine/Plugins/Enterprise/DatasmithImporter/Source/DirectLinkExtension/Private/DirectLinkManager.cpp
--- a/Engine/Plugins/Enterprise/DatasmithImporter/Source/DirectLinkExtension/Private/DirectLinkManager.cpp
+++ b/Engine/Plugins/Enterprise/DatasmithImporter/Source/DirectLinkExtension/Private/DirectLinkManager.cpp
@@ -235,7 +235,7 @@ namespace UE::DatasmithImporter
 					&& !ExternalSource->GetDatasmithScene().IsValid())
 				{
 					ExternalSource->CancelAsyncLoad();
-					UE_LOG(LogDirectLinkManager, Error, TEXT("The DirectLink source \"%s\" could not be loaded: Nothing to synchronize. Make sure to do a DirectLink sync in your exporter."), *ExternalSource->GetSourceName());
+					UE_LOG(LogDirectLinkManager, Warning, TEXT("The DirectLink source \"%s\" could not be loaded: Nothing to synchronize. Make sure to do a DirectLink sync in your exporter."), *ExternalSource->GetSourceName());
 				}
 			}
 		}
diff --git a/Engine/Plugins/Enterprise/DatasmithImporter/Source/ExternalSource/Public/ExternalSource.h b/Engine/Plugins/Enterprise/DatasmithImporter/Source/ExternalSource/Public/ExternalSource.h
--- a/Engine/Plugins/Enterprise/DatasmithImporter/Source/ExternalSource/Public/ExternalSource.h
+++ b/Engine/Plugins/Enterprise/DatasmithImporter/Source/ExternalSource/Public/ExternalSource.h
@@ -125,7 +125,7 @@ namespace UE::DatasmithImporter
 		/**
 		 * True when the ExternalSource is currently loading the source asynchronously.
 		 */
-		bool IsAsyncLoading() const { return !PendingPromiseQueue.IsEmpty(); }
+		bool IsAsyncLoading() const { return AsyncTaskNotification.IsValid(); }
 
 		/**
 		 * Cancel any pending AsyncLoad() operation, pending TFuture will return invalid TSharedPtr.
