Non-unity build missing includes

diff --git a/Engine/Plugins/Enterprise/DatasmithImporter/Source/DirectLinkExtension/Private/DirectLinkManager.h b/Engine/Plugins/Enterprise/DatasmithImporter/Source/DirectLinkExtension/Private/DirectLinkManager.h
--- a/Engine/Plugins/Enterprise/DatasmithImporter/Source/DirectLinkExtension/Private/DirectLinkManager.h
+++ b/Engine/Plugins/Enterprise/DatasmithImporter/Source/DirectLinkExtension/Private/DirectLinkManager.h
@@ -7,6 +7,7 @@
 #include "SourceUri.h"
 #include "Containers/Map.h"
 #include "Containers/UnrealString.h"
+#include "Containers/Queue.h"
 #include "Delegates/DelegateCombinations.h"
 #include "DirectLinkEndpoint.h"
 #include "HAL/CriticalSection.h"
