DirectLink in Editor - Preserve the DirectLink source info if "reimport from new file" select the Datasmith.directlink file

JeanLuc.Corenthin

diff --git a/Engine/Plugins/Enterprise/DatasmithImporter/Source/DatasmithImporter/Private/DatasmithImportFactory.cpp b/Engine/Plugins/Enterprise/DatasmithImporter/Source/DatasmithImporter/Private/DatasmithImportFactory.cpp
--- a/Engine/Plugins/Enterprise/DatasmithImporter/Source/DatasmithImporter/Private/DatasmithImportFactory.cpp
+++ b/Engine/Plugins/Enterprise/DatasmithImporter/Source/DatasmithImporter/Private/DatasmithImportFactory.cpp
@@ -620,16 +620,31 @@ void UDatasmithImportFactory::SetReimportPaths( UObject* Obj, const TArray<FStri
 
 	if (UAssetImportData* ReImportData = DatasmithImportFactoryImpl::GetImportData(Obj))
 	{
+		const FString PreviousImportPath = ReImportData->GetFirstFilename();
 		ReImportData->UpdateFilenameOnly(NewReimportPaths[0]);
 
-		FString SourceUriString = FSourceUri::FromFilePath(NewReimportPaths[0]).ToString();
+		FDatasmithImportInfo* DatasmithImportInfo = nullptr;
 		if (UDatasmithAssetImportData* DatasmithAssetReImportData = Cast<UDatasmithAssetImportData>(ReImportData))
 		{
-			DatasmithAssetReImportData->DatasmithImportInfo = FDatasmithImportInfo(SourceUriString);
+			DatasmithImportInfo = &DatasmithAssetReImportData->DatasmithImportInfo;
 		}
 		else if (UDatasmithSceneImportData* DatasmithSceneReImportData = Cast<UDatasmithSceneImportData>(ReImportData))
 		{
-			DatasmithSceneReImportData->DatasmithImportInfo = FDatasmithImportInfo(SourceUriString);
+			DatasmithImportInfo = &DatasmithSceneReImportData->DatasmithImportInfo;
+		}
+
+		if (DatasmithImportInfo)
+		{
+			const bool bWasFileSource = FSourceUri(DatasmithImportInfo->SourceUri).GetScheme() == FSourceUri::GetFileScheme();
+			const bool bChangedFilename = FPaths::ConvertRelativePathToFull(PreviousImportPath) != FPaths::ConvertRelativePathToFull(NewReimportPaths[0]);
+
+			// Workaround to avoid clearing the import info on sources (directlink) that were not imported from a file.
+			// #ueent_todo The ReimportFromNewFile command should be replaced by a ReimportFromNewSource, adding support for non-file sources.
+			if (bWasFileSource || bChangedFilename)
+			{
+				const FString SourceUriString = FSourceUri::FromFilePath(NewReimportPaths[0]).ToString();
+				*DatasmithImportInfo = FDatasmithImportInfo(SourceUriString);
+			}
 		}
 	}
 }
