Non-Unity build fixes

diff --git a/Engine/Plugins/Enterprise/DatasmithImporter/Source/DirectLinkExtension/Private/UI/SDirectLinkAssetThumbnailIndicator.cpp b/Engine/Plugins/Enterprise/DatasmithImporter/Source/DirectLinkExtension/Private/UI/SDirectLinkAssetThumbnailIndicator.cpp
--- a/Engine/Plugins/Enterprise/DatasmithImporter/Source/DirectLinkExtension/Private/UI/SDirectLinkAssetThumbnailIndicator.cpp
+++ b/Engine/Plugins/Enterprise/DatasmithImporter/Source/DirectLinkExtension/Private/UI/SDirectLinkAssetThumbnailIndicator.cpp
@@ -8,6 +8,7 @@
 #include "ExternalSourceModule.h"
 #include "SourceUri.h"
 #include "Widgets/Images/SImage.h"
+#include "Widgets/Text/STextBlock.h"
 #include "Widgets/SOverlay.h"
 
 #define LOCTEXT_NAMESPACE "DirectLinkAssetThumbnailIndicator"
diff --git a/Engine/Plugins/MovieScene/LevelSequenceEditor/Source/LevelSequenceEditor/Private/LevelSequenceEditorSubsystem.cpp b/Engine/Plugins/MovieScene/LevelSequenceEditor/Source/LevelSequenceEditor/Private/LevelSequenceEditorSubsystem.cpp
--- a/Engine/Plugins/MovieScene/LevelSequenceEditor/Source/LevelSequenceEditor/Private/LevelSequenceEditorSubsystem.cpp
+++ b/Engine/Plugins/MovieScene/LevelSequenceEditor/Source/LevelSequenceEditor/Private/LevelSequenceEditorSubsystem.cpp
@@ -12,7 +12,8 @@
 #include "LevelEditor.h"
 #include "ScopedTransaction.h"
 #include "Compilation/MovieSceneCompiledDataManager.h"
-
+#include "Modules/ModuleManager.h"
+#include "Framework/MultiBox/MultiBoxBuilder.h"
 
 DEFINE_LOG_CATEGORY(LogLevelSequenceEditor);
 
diff --git a/Engine/Source/Editor/EditorFramework/Private/Viewports/InViewportUIDragOperation.cpp b/Engine/Source/Editor/EditorFramework/Private/Viewports/InViewportUIDragOperation.cpp
--- a/Engine/Source/Editor/EditorFramework/Private/Viewports/InViewportUIDragOperation.cpp
+++ b/Engine/Source/Editor/EditorFramework/Private/Viewports/InViewportUIDragOperation.cpp
@@ -4,6 +4,7 @@
 #include "Widgets/Docking/SDockTab.h"
 #include "Widgets/SWindow.h"
 #include "Styling/AppStyle.h"
+#include "Framework/Application/SlateApplication.h"
 
 const FVector2D FInViewportUIDragOperation::GetDecoratorOffsetFromCursor()
 {
diff --git a/Engine/Source/Runtime/Engine/Private/WorldPartition/ErrorHandling/WorldPartitionStreamingGenerationMapCheckErrorHandler.cpp b/Engine/Source/Runtime/Engine/Private/WorldPartition/ErrorHandling/WorldPartitionStreamingGenerationMapCheckErrorHandler.cpp
--- a/Engine/Source/Runtime/Engine/Private/WorldPartition/ErrorHandling/WorldPartitionStreamingGenerationMapCheckErrorHandler.cpp
+++ b/Engine/Source/Runtime/Engine/Private/WorldPartition/ErrorHandling/WorldPartitionStreamingGenerationMapCheckErrorHandler.cpp
@@ -5,6 +5,7 @@
 #include "Logging/MessageLog.h"
 #include "Misc/UObjectToken.h"
 #include "Misc/MapErrors.h"
+#include "GameFramework/Actor.h"
 
 #define LOCTEXT_NAMESPACE "WorldPartition"
 
