Datasmith: Material hash : also hash free expressions not connected to the main material node.
The recursive algo now keeps track of visited expressions in order to find "free" expressions that are not conneced to the main output node, but should still be part of the hash.

diff --git a/Engine/Plugins/Enterprise/DatasmithImporter/Source/DatasmithImporter/Private/DatasmithMaterialImporter.cpp b/Engine/Plugins/Enterprise/DatasmithImporter/Source/DatasmithImporter/Private/DatasmithMaterialImporter.cpp
--- a/Engine/Plugins/Enterprise/DatasmithImporter/Source/DatasmithImporter/Private/DatasmithMaterialImporter.cpp
+++ b/Engine/Plugins/Enterprise/DatasmithImporter/Source/DatasmithImporter/Private/DatasmithMaterialImporter.cpp
@@ -27,9 +27,9 @@
 
 namespace DatasmithMaterialImporterUtils
 {
-	int32 ComputeMaterialExpressionHash( IDatasmithMaterialExpression* MaterialExpression );
+	int32 ComputeMaterialExpressionHash( IDatasmithMaterialExpression* MaterialExpression, TSet<IDatasmithMaterialExpression*>& VisitedExpressions );
 
-	int32 ComputeExpressionInputHash( IDatasmithExpressionInput* ExpressionInput )
+	int32 ComputeExpressionInputHash( IDatasmithExpressionInput* ExpressionInput, TSet<IDatasmithMaterialExpression*>& VisitedExpressions )
 	{
 		if ( !ExpressionInput )
 		{
@@ -40,7 +40,7 @@ namespace DatasmithMaterialImporterUtils
 
 		if ( ExpressionInput->GetExpression() )
 		{
-			int32 ExpressionHash = ComputeMaterialExpressionHash( ExpressionInput->GetExpression() );
+			int32 ExpressionHash = ComputeMaterialExpressionHash( ExpressionInput->GetExpression(), VisitedExpressions );
 			Hash = HashCombine( Hash, ExpressionHash );
 		}
 
@@ -49,8 +49,9 @@ namespace DatasmithMaterialImporterUtils
 		return Hash;
 	}
 
-	int32 ComputeMaterialExpressionHash( IDatasmithMaterialExpression* MaterialExpression )
+	int32 ComputeMaterialExpressionHash( IDatasmithMaterialExpression* MaterialExpression, TSet<IDatasmithMaterialExpression*>& VisitedExpressions )
 	{
+		VisitedExpressions.Add(MaterialExpression);
 		uint32 Hash = GetTypeHash( MaterialExpression->GetExpressionType() );
 		Hash = HashCombine( Hash, GetTypeHash( FString( MaterialExpression->GetName() ) ) );
 
@@ -120,7 +121,7 @@ namespace DatasmithMaterialImporterUtils
 
 		for ( int32 InputIndex = 0; InputIndex < MaterialExpression->GetInputCount(); ++InputIndex )
 		{
-			Hash = HashCombine( Hash, ComputeExpressionInputHash( MaterialExpression->GetInput( InputIndex ) ) );
+			Hash = HashCombine( Hash, ComputeExpressionInputHash( MaterialExpression->GetInput( InputIndex ), VisitedExpressions ) );
 		}
 
 		return Hash;
@@ -128,22 +129,40 @@ namespace DatasmithMaterialImporterUtils
 
 	int32 ComputeMaterialHash( TSharedRef< IDatasmithUEPbrMaterialElement > MaterialElement )
 	{
+		TSet<IDatasmithMaterialExpression*> VisitedExpressions;
 		int32 Hash = GetTypeHash( MaterialElement->GetTwoSided() );
 
 		Hash = HashCombine( Hash, GetTypeHash( MaterialElement->GetUseMaterialAttributes() ) );
 		Hash = HashCombine( Hash, GetTypeHash( MaterialElement->GetBlendMode() ) );
 		Hash = HashCombine( Hash, GetTypeHash( MaterialElement->GetShadingModel() ) );
 
-		Hash = HashCombine( Hash, ComputeExpressionInputHash( &MaterialElement->GetBaseColor() ) );
-		Hash = HashCombine( Hash, ComputeExpressionInputHash( &MaterialElement->GetMetallic() ) );
-		Hash = HashCombine( Hash, ComputeExpressionInputHash( &MaterialElement->GetSpecular() ) );
-		Hash = HashCombine( Hash, ComputeExpressionInputHash( &MaterialElement->GetRoughness() ) );
-		Hash = HashCombine( Hash, ComputeExpressionInputHash( &MaterialElement->GetEmissiveColor() ) );
-		Hash = HashCombine( Hash, ComputeExpressionInputHash( &MaterialElement->GetOpacity() ) );
-		Hash = HashCombine( Hash, ComputeExpressionInputHash( &MaterialElement->GetNormal() ) );
-		Hash = HashCombine( Hash, ComputeExpressionInputHash( &MaterialElement->GetRefraction() ) );
-		Hash = HashCombine( Hash, ComputeExpressionInputHash( &MaterialElement->GetAmbientOcclusion() ) );
-		Hash = HashCombine( Hash, ComputeExpressionInputHash( &MaterialElement->GetMaterialAttributes() ) );
+		Hash = HashCombine( Hash, ComputeExpressionInputHash( &MaterialElement->GetBaseColor(), VisitedExpressions ) );
+		Hash = HashCombine( Hash, ComputeExpressionInputHash( &MaterialElement->GetMetallic(), VisitedExpressions ) );
+		Hash = HashCombine( Hash, ComputeExpressionInputHash( &MaterialElement->GetSpecular(), VisitedExpressions ) );
+		Hash = HashCombine( Hash, ComputeExpressionInputHash( &MaterialElement->GetRoughness(), VisitedExpressions ) );
+		Hash = HashCombine( Hash, ComputeExpressionInputHash( &MaterialElement->GetEmissiveColor(), VisitedExpressions ) );
+		Hash = HashCombine( Hash, ComputeExpressionInputHash( &MaterialElement->GetOpacity(), VisitedExpressions ) );
+		Hash = HashCombine( Hash, ComputeExpressionInputHash( &MaterialElement->GetNormal(), VisitedExpressions ) );
+		Hash = HashCombine( Hash, ComputeExpressionInputHash( &MaterialElement->GetRefraction(), VisitedExpressions ) );
+		Hash = HashCombine( Hash, ComputeExpressionInputHash( &MaterialElement->GetAmbientOcclusion(), VisitedExpressions ) );
+		Hash = HashCombine( Hash, ComputeExpressionInputHash( &MaterialElement->GetMaterialAttributes(), VisitedExpressions ) );
+
+		TSet<IDatasmithMaterialExpression*> AllExpressions;
+		for (int32 ExpressionIndex = 0; ExpressionIndex < MaterialElement->GetExpressionsCount(); ++ExpressionIndex)
+		{
+			AllExpressions.Add(MaterialElement->GetExpression(ExpressionIndex));
+		}
+
+		TSet<IDatasmithMaterialExpression*> NotVisitedExpressions = AllExpressions.Difference(VisitedExpressions);
+		NotVisitedExpressions.Remove(nullptr);
+		// we don't want to depend on the order of nodes
+		NotVisitedExpressions.Sort([](IDatasmithMaterialExpression& A, IDatasmithMaterialExpression& B){ return FCString::Strcmp(A.GetName(), B.GetName()) < 0; });
+
+		for (IDatasmithMaterialExpression* Expression : NotVisitedExpressions)
+		{
+			TSet<IDatasmithMaterialExpression*> _;
+			Hash = HashCombine(Hash, ComputeMaterialExpressionHash(Expression, _));
+		}
 
 		return Hash;
 	}
