Removing GenerateSceneSnapshotHash() from DirectLink public API. This should only be used internally until we come up with a more scallable way of detecting snapshot changes.

diff --git a/Engine/Plugins/Enterprise/DatasmithImporter/Source/DirectLinkExtension/Private/DirectLinkExternalSource.cpp b/Engine/Plugins/Enterprise/DatasmithImporter/Source/DirectLinkExtension/Private/DirectLinkExternalSource.cpp
--- a/Engine/Plugins/Enterprise/DatasmithImporter/Source/DirectLinkExtension/Private/DirectLinkExternalSource.cpp
+++ b/Engine/Plugins/Enterprise/DatasmithImporter/Source/DirectLinkExtension/Private/DirectLinkExternalSource.cpp
@@ -9,7 +9,9 @@
 #include "Async/Async.h"
 #include "DirectLinkEndpoint.h"
 #include "DirectLinkDeltaConsumer.h"
+#include "DirectLinkElementSnapshot.h"
 #include "DirectLinkMisc.h"
+#include "DirectLinkSceneSnapshot.h"
 #include "Misc/AsyncTaskNotification.h"
 
 #define LOCTEXT_NAMESPACE "DirectLinkExternalSource"
@@ -32,12 +34,31 @@ namespace UE::DatasmithImporter
 			SceneReceiver->FinalSnapshot(SceneSnapshot);
 			if (TSharedPtr<FDirectLinkExternalSource> PinnedExternalSource = DirectLinkExternalSource.Pin())
 			{
-				PinnedExternalSource->CachedHash = DirectLink::GenerateSceneSnapshotHash(SceneSnapshot);
+				PinnedExternalSource->CachedHash = GenerateSceneSnapshotHash(SceneSnapshot);
 				PinnedExternalSource->TriggerOnExternalSourceChanged();
 			}
 		}
 
 	private:
+
+		FMD5Hash GenerateSceneSnapshotHash(const DirectLink::FSceneSnapshot& SceneSnapshot)
+		{
+			using namespace DirectLink;
+			FMD5Hash Hash;
+			FMD5 SceneMD5Hash;
+
+			SceneMD5Hash.Update((uint8*)&SceneSnapshot.SceneId, sizeof(FSceneGraphId));
+			for (const TPair<FSceneGraphId, TSharedRef<FElementSnapshot>>& ElementPair : SceneSnapshot.Elements)
+			{
+				FElementHash ElementHash(ElementPair.Value->GetHash());
+				SceneMD5Hash.Update((uint8*)&ElementPair.Key, sizeof(FSceneGraphId));
+				SceneMD5Hash.Update((uint8*)&ElementHash, sizeof(FElementHash));
+			}
+
+			Hash.Set(SceneMD5Hash);
+			return Hash;
+		}
+
 		TSharedRef<DirectLink::ISceneReceiver> SceneReceiver;
 		TWeakPtr<FDirectLinkExternalSource> DirectLinkExternalSource;
 	};
diff --git a/Engine/Source/Runtime/Datasmith/DirectLink/Private/DirectLinkMisc.cpp b/Engine/Source/Runtime/Datasmith/DirectLink/Private/DirectLinkMisc.cpp
--- a/Engine/Source/Runtime/Datasmith/DirectLink/Private/DirectLinkMisc.cpp
+++ b/Engine/Source/Runtime/Datasmith/DirectLink/Private/DirectLinkMisc.cpp
@@ -76,21 +76,4 @@ void DumpSceneSnapshot(FSceneSnapshot& SceneSnapshot, const FString& BaseFileNam
 	}
 }
 
-FMD5Hash GenerateSceneSnapshotHash(const FSceneSnapshot& SceneSnapshot)
-{
-	FMD5Hash Hash;
-	FMD5 SceneMD5Hash;
-
-	SceneMD5Hash.Update((uint8*)&SceneSnapshot.SceneId, sizeof(FSceneGraphId));
-	for (const TPair<FSceneGraphId, TSharedRef<FElementSnapshot>>& ElementPair : SceneSnapshot.Elements)
-	{
-		FElementHash ElementHash(ElementPair.Value->GetHash());
-		SceneMD5Hash.Update((uint8*)&ElementPair.Key, sizeof(FSceneGraphId));
-		SceneMD5Hash.Update((uint8*)&ElementHash, sizeof(FElementHash));
-	}
-
-	Hash.Set(SceneMD5Hash);
-	return Hash;
-}
-
 } // namespace DirectLink
diff --git a/Engine/Source/Runtime/Datasmith/DirectLink/Public/DirectLinkMisc.h b/Engine/Source/Runtime/Datasmith/DirectLink/Public/DirectLinkMisc.h
--- a/Engine/Source/Runtime/Datasmith/DirectLink/Public/DirectLinkMisc.h
+++ b/Engine/Source/Runtime/Datasmith/DirectLink/Public/DirectLinkMisc.h
@@ -3,7 +3,6 @@
 #pragma once
 
 #include "CoreTypes.h"
-#include "Misc/SecureHash.h"
 #include "Templates/SharedPointer.h"
 
 
@@ -17,7 +16,4 @@ class FSceneSnapshot;
 DIRECTLINK_API const FString& GetDumpPath();
 
 DIRECTLINK_API void DumpSceneSnapshot(FSceneSnapshot& SceneSnapshot, const FString& BaseFileName);
-
-DIRECTLINK_API FMD5Hash GenerateSceneSnapshotHash(const FSceneSnapshot& SceneSnapshot);
-
 } // namespace DirectLink
