FSourceUri cleaning up the query constructor.

diff --git a/Engine/Plugins/Enterprise/DatasmithImporter/Source/ExternalSource/Public/SourceUri.h b/Engine/Plugins/Enterprise/DatasmithImporter/Source/ExternalSource/Public/SourceUri.h
--- a/Engine/Plugins/Enterprise/DatasmithImporter/Source/ExternalSource/Public/SourceUri.h
+++ b/Engine/Plugins/Enterprise/DatasmithImporter/Source/ExternalSource/Public/SourceUri.h
@@ -31,32 +31,8 @@ namespace UE::DatasmithImporter
 
 		template<typename TContainer>
 		FSourceUri(const FString& InScheme, const FString& InPath, const TContainer& QueryContainer)
-			: Uri(InScheme + GetSchemeDelimiter() + InPath)
-		{
-			static_assert(std::is_same<typename TContainer::ElementType, TPair<FString, FString>>::value, TEXT("Query container must consist of a collection of TPair<FString, FString>."));
-
-			if (QueryContainer.Num() > 0)
-			{
-				TCHAR QueryDelimiter(GetQueryStartCharacter());
-				FString QueryString(1, &QueryDelimiter);
-				for (typename TContainer::TConstIterator It = QueryContainer.CreateConstIterator(); It;)
-				{
-					const TPair<FString, FString>& CurrentPair = *It;
-
-					if (++It)
-					{
-						// We are not on the last pair, add '&' character to split this pair with the next.
-						QueryString += FString::Printf(TEXT("%s%c%s%c"), *CurrentPair.Key, GetQueryPairCharacter(), *CurrentPair.Value, GetQueryAppendPairCharacter());
-					}
-					else
-					{
-						QueryString += FString::Printf(TEXT("%s%c%s"), *CurrentPair.Key, GetQueryPairCharacter(), *CurrentPair.Value);
-					}
-				}
-
-				Uri += QueryString;
-			}
-		}
+			: Uri(InScheme + GetSchemeDelimiter() + InPath + BuildQueryString(QueryContainer))
+		{}
 
 		/**
 		 * Generate a FSourceUri from the given filepath.
@@ -119,6 +95,36 @@ namespace UE::DatasmithImporter
 
 	private:
 
+		template<typename TContainer>
+		static FString BuildQueryString(const TContainer& QueryContainer)
+		{
+			static_assert(std::is_same<typename TContainer::ElementType, TPair<FString, FString>>::value, TEXT("Query container must consist of a collection of TPair<FString, FString>."));
+
+			if (QueryContainer.Num() > 0)
+			{
+				TCHAR QueryDelimiter(GetQueryStartCharacter());
+				FString QueryString(1, &QueryDelimiter);
+				for (typename TContainer::TConstIterator It = QueryContainer.CreateConstIterator(); It;)
+				{
+					const TPair<FString, FString>& CurrentPair = *It;
+
+					if (++It)
+					{
+						// We are not on the last pair, add '&' character to split this pair with the next.
+						QueryString += FString::Printf(TEXT("%s%c%s%c"), *CurrentPair.Key, GetQueryPairCharacter(), *CurrentPair.Value, GetQueryAppendPairCharacter());
+					}
+					else
+					{
+						QueryString += FString::Printf(TEXT("%s%c%s"), *CurrentPair.Key, GetQueryPairCharacter(), *CurrentPair.Value);
+					}
+				}
+
+				return QueryString;
+			}
+
+			return FString();
+		}
+
 		static const FString& GetSchemeDelimiter();
 
 		/**
