Datasmith : fix Environnment map binding issue

diff --git a/Engine/Plugins/Enterprise/DatasmithImporter/Source/DatasmithImporter/Private/DatasmithActorImporter.cpp b/Engine/Plugins/Enterprise/DatasmithImporter/Source/DatasmithImporter/Private/DatasmithActorImporter.cpp
--- a/Engine/Plugins/Enterprise/DatasmithImporter/Source/DatasmithImporter/Private/DatasmithActorImporter.cpp
+++ b/Engine/Plugins/Enterprise/DatasmithImporter/Source/DatasmithImporter/Private/DatasmithActorImporter.cpp
@@ -477,8 +477,8 @@ AActor* FDatasmithActorImporter::ImportEnvironment(FDatasmithImportContext& Impo
 				ExistingMaterial = FindObject<UMaterial>(ImportContext.AssetsContext.MaterialsFinalPackage.Get(), ShaderElement->GetName());
 			}
 
-			FDatasmithMaterialExpressions::CreateDatasmithEnvironmentMaterial(ImportContext.AssetsContext.MaterialsFinalPackage.Get(), ShaderElement, ImportContext.AssetsContext, ExistingMaterial);
-			UMaterialInterface* Material = FDatasmithImporterUtils::FindAsset< UMaterialInterface >( ImportContext.AssetsContext, ShaderElement->GetName() );
+			UPackage* Package = CreatePackage( *FPaths::Combine( ImportContext.AssetsContext.MaterialsFinalPackage.Get()->GetPathName(), ShaderElement->GetName() ) );
+			UMaterialInterface* Material = FDatasmithMaterialExpressions::CreateDatasmithEnvironmentMaterial(Package, ShaderElement, ImportContext.AssetsContext, ExistingMaterial);
 			if (Material)
 			{
 				EnvironmentActor->GetStaticMeshComponent()->SetMaterial(0, Material);
