DirectLink Editor - Adding BP interface for importing via DirectLink external source

Johan.Duparc

diff --git a/Engine/Plugins/Enterprise/DatasmithImporter/Source/DatasmithImporter/Private/DatasmithBlueprintLibrary.cpp b/Engine/Plugins/Enterprise/DatasmithImporter/Source/DatasmithImporter/Private/DatasmithBlueprintLibrary.cpp
--- a/Engine/Plugins/Enterprise/DatasmithImporter/Source/DatasmithImporter/Private/DatasmithBlueprintLibrary.cpp
+++ b/Engine/Plugins/Enterprise/DatasmithImporter/Source/DatasmithImporter/Private/DatasmithBlueprintLibrary.cpp
@@ -19,6 +19,8 @@
 #include "Async/ParallelFor.h"
 #include "AssetRegistryModule.h"
 #include "DatasmithAssetImportData.h"
+#include "DirectLinkExtensionModule.h"
+#include "DirectLinkExternalSource.h"
 #include "Editor.h"
 #include "EditorAssetLibrary.h"
 #include "Engine/StaticMesh.h"
@@ -543,14 +545,22 @@ namespace DatasmithSceneElementUtil
 
 UDatasmithSceneElement* UDatasmithSceneElement::ConstructDatasmithSceneFromFile(const FString& InFilename)
 {
-	using namespace DatasmithBlueprintLibraryImpl;
 	using namespace UE::DatasmithImporter;
 
 	const FSourceUri SourceUri = FSourceUri::FromFilePath(InFilename);
+	return ConstructDatasmithSceneFromSourceUri(SourceUri.ToString());
+}
+
+UDatasmithSceneElement* UDatasmithSceneElement::ConstructDatasmithSceneFromSourceUri(const FString& SourceUriString)
+{
+	using namespace DatasmithBlueprintLibraryImpl;
+	using namespace UE::DatasmithImporter;
+
+	const FSourceUri SourceUri = FSourceUri(SourceUriString);
 	TSharedPtr<FExternalSource> ExternalSource = IExternalSourceModule::GetOrCreateExternalSource(SourceUri);
 	if (!ExternalSource)
 	{
-		UE_LOG(LogDatasmithImport, Error, TEXT("Datasmith import error: file name could not be resolved to an external source. Abort import."));
+		UE_LOG(LogDatasmithImport, Error, TEXT("Datasmith import error: source uri could not be resolved to an external source. Abort import."));
 		return nullptr;
 	}
 
diff --git a/Engine/Plugins/Enterprise/DatasmithImporter/Source/DatasmithImporter/Public/DatasmithBlueprintLibrary.h b/Engine/Plugins/Enterprise/DatasmithImporter/Source/DatasmithImporter/Public/DatasmithBlueprintLibrary.h
--- a/Engine/Plugins/Enterprise/DatasmithImporter/Source/DatasmithImporter/Public/DatasmithBlueprintLibrary.h
+++ b/Engine/Plugins/Enterprise/DatasmithImporter/Source/DatasmithImporter/Public/DatasmithBlueprintLibrary.h
@@ -56,6 +56,14 @@ public:
 	UFUNCTION(BlueprintCallable, Category = "Editor Scripting | Datasmith")
 	static UDatasmithSceneElement* ConstructDatasmithSceneFromFile(const FString& FilePath);
 
+	/**
+	 * Open an existing Datasmith source from the SourceUri.
+	 * @param	SourceUri Uri of the datasmith scene to open. ie: file://c:/MyFolder/MyFiles.udatasmith
+	 * @return	The opened DatasmithScene, that can be modified and can be imported.
+	 **/
+	UFUNCTION(BlueprintCallable, Category = "Editor Scripting | Datasmith")
+	static UDatasmithSceneElement* ConstructDatasmithSceneFromSourceUri(const FString& SourceUri);
+
 	/**
 	 * Open set of CAD files as actors in a single datasmith scene
 	 * Importing set of files into single DatasmithScene asset(with ImportScene) is supported only for CAD files
diff --git a/Engine/Plugins/Enterprise/DatasmithImporter/Source/DirectLinkExtension/Private/DirectLinkExtensionBlueprintLibrary.cpp b/Engine/Plugins/Enterprise/DatasmithImporter/Source/DirectLinkExtension/Private/DirectLinkExtensionBlueprintLibrary.cpp
new file mode 100644
--- /dev/null
+++ b/Engine/Plugins/Enterprise/DatasmithImporter/Source/DirectLinkExtension/Private/DirectLinkExtensionBlueprintLibrary.cpp
@@ -0,0 +1,43 @@
+// Copyright Epic Games, Inc. All Rights Reserved.
+
+#include "DirectLinkExtensionBlueprintLibrary.h"
+
+#include "DirectLinkExtensionModule.h"
+#include "DirectLinkExternalSource.h"
+#include "DirectLinkUriResolver.h"
+#include "IDirectLinkManager.h"
+#include "SourceUri.h"
+
+TArray<FString> UDirectLinkExtensionBlueprintLibrary::GetAvailableDirectLinkSourcesUri()
+{
+	using namespace UE::DatasmithImporter;
+
+	TArray<FString> SourceUris;
+	const IDirectLinkManager& Manager = IDirectLinkExtensionModule::Get().GetManager();
+	const TArray<TSharedRef<FDirectLinkExternalSource>> ExternalSources = Manager.GetExternalSourceList();
+
+	for (const TSharedRef<FDirectLinkExternalSource>& ExternalSource : ExternalSources)
+	{
+		SourceUris.Add(ExternalSource->GetSourceUri().ToString());
+	}
+
+	return SourceUris;
+}
+
+bool UDirectLinkExtensionBlueprintLibrary::ParseDirectLinkSourceUri(const FString& SourceUriString, FString& OutComputerName, FString& OutEndpointName, FString& OutExecutableName, FString& OutSourceName)
+{
+	using namespace UE::DatasmithImporter;
+
+	const FSourceUri SourceUri(SourceUriString);
+	if (TOptional<FDirectLinkSourceDescription> SourceDescription = FDirectLinkUriResolver::TryParseDirectLinkUri(SourceUri))
+	{
+		OutComputerName = SourceDescription->ComputerName;
+		OutEndpointName = SourceDescription->EndpointName;
+		OutExecutableName = SourceDescription->ExecutableName;
+		OutSourceName = SourceDescription->SourceName;
+
+		return true;
+	}
+
+	return false;
+}
diff --git a/Engine/Plugins/Enterprise/DatasmithImporter/Source/DirectLinkExtension/Public/DirectLinkExtensionBlueprintLibrary.h b/Engine/Plugins/Enterprise/DatasmithImporter/Source/DirectLinkExtension/Public/DirectLinkExtensionBlueprintLibrary.h
new file mode 100644
--- /dev/null
+++ b/Engine/Plugins/Enterprise/DatasmithImporter/Source/DirectLinkExtension/Public/DirectLinkExtensionBlueprintLibrary.h
@@ -0,0 +1,22 @@
+// Copyright Epic Games, Inc. All Rights Reserved.
+
+#pragma once
+
+#include "Containers/Array.h"
+#include "Containers/UnrealString.h"
+#include "Kismet/BlueprintFunctionLibrary.h"
+
+#include "DirectLinkExtensionBlueprintLibrary.generated.h"
+
+UCLASS(MinimalAPI, meta = (ScriptName = "DirectLinkExtensionLibrary"))
+class UDirectLinkExtensionBlueprintLibrary : public UBlueprintFunctionLibrary
+{
+	GENERATED_BODY()
+
+public:
+	UFUNCTION(BlueprintCallable, Category = "Editor Scripting | DirectLink")
+	static TArray<FString> GetAvailableDirectLinkSourcesUri();
+
+	UFUNCTION(BlueprintCallable, Category = "Editor Scripting | DirectLink")
+	static bool ParseDirectLinkSourceUri(const FString& SourceUriString, FString& OutComputerName, FString& OutEndpointName, FString& OutExecutableName, FString& OutSourceName);
+};
\ No newline at end of file
