Datasmith: fix missing extensions listing in various import/reimport dialogs

diff --git a/Engine/Plugins/Enterprise/DatasmithImporter/Source/DatasmithImporter/Private/DatasmithImportFactory.cpp b/Engine/Plugins/Enterprise/DatasmithImporter/Source/DatasmithImporter/Private/DatasmithImportFactory.cpp
--- a/Engine/Plugins/Enterprise/DatasmithImporter/Source/DatasmithImporter/Private/DatasmithImportFactory.cpp
+++ b/Engine/Plugins/Enterprise/DatasmithImporter/Source/DatasmithImporter/Private/DatasmithImportFactory.cpp
@@ -578,9 +578,11 @@ void UDatasmithImportFactory::ParseFromJson(TSharedRef<FJsonObject> InImportSett
 bool UDatasmithImportFactory::CanReimport( UObject* Obj, TArray<FString>& OutFilenames )
 {
 	// The CDO may be used to do that check. In that case, Formats are not necessarily initialized.
-	const bool bGenerateSupportedFormats = HasAnyFlags(RF_ClassDefaultObject) && Formats.Num() == 0;
-	TGuardValue<TArray<FString>> FormatsGuard(Formats, bGenerateSupportedFormats ? FDatasmithTranslatorManager::Get().GetSupportedFormats() : Formats);
-	
+	if (Formats.IsEmpty())
+	{
+		Formats = FDatasmithTranslatorManager::Get().GetSupportedFormats();
+	}
+
 	UAssetImportData* ReImportData = DatasmithImportFactoryImpl::GetImportData(Obj);
 	if (ReImportData == nullptr)
 	{
