DirectLink in Editor - Reconnecting external source after connection lost.

Johan.Duparc

diff --git a/Engine/Plugins/Enterprise/DatasmithImporter/Source/DirectLinkExtension/Private/DirectLinkExternalSource.cpp b/Engine/Plugins/Enterprise/DatasmithImporter/Source/DirectLinkExtension/Private/DirectLinkExternalSource.cpp
--- a/Engine/Plugins/Enterprise/DatasmithImporter/Source/DirectLinkExtension/Private/DirectLinkExternalSource.cpp
+++ b/Engine/Plugins/Enterprise/DatasmithImporter/Source/DirectLinkExtension/Private/DirectLinkExternalSource.cpp
@@ -85,11 +85,14 @@ namespace UE::DatasmithImporter
 
 	bool FDirectLinkExternalSource::OpenStream()
 	{
-		if (SourceHandle.IsValid() && DestinationHandle.IsValid() && !bIsStreamOpen)
+		using namespace DirectLink;
+
+		if (SourceHandle.IsValid() && DestinationHandle.IsValid())
 		{
-			const DirectLink::FEndpoint::EOpenStreamResult Result = IDirectLinkExtensionModule::GetEndpoint().OpenStream(SourceHandle, DestinationHandle);
+			const FEndpoint::EOpenStreamResult Result = IDirectLinkExtensionModule::GetEndpoint().OpenStream(SourceHandle, DestinationHandle);
+
+			bIsStreamOpen = Result == FEndpoint::EOpenStreamResult::Opened || Result == FEndpoint::EOpenStreamResult::AlreadyOpened;
 
-			bIsStreamOpen = Result == DirectLink::FEndpoint::EOpenStreamResult::Opened;
 			return bIsStreamOpen;			
 		}
 
@@ -98,7 +101,7 @@ namespace UE::DatasmithImporter
 
 	void FDirectLinkExternalSource::CloseStream()
 	{
-		if (SourceHandle.IsValid() && DestinationHandle.IsValid() && bIsStreamOpen)
+		if (SourceHandle.IsValid() && DestinationHandle.IsValid())
 		{
 			IDirectLinkExtensionModule::GetEndpoint().CloseStream(SourceHandle, DestinationHandle);
 			bIsStreamOpen = false;
@@ -109,8 +112,8 @@ namespace UE::DatasmithImporter
 	{
 		if (IDirectLinkExtensionModule::IsAvailable())
 		{
-			IDirectLinkExtensionModule::GetEndpoint().RemoveDestination(DestinationHandle);
 			CloseStream();
+			IDirectLinkExtensionModule::GetEndpoint().RemoveDestination(DestinationHandle);
 		}
 
 		ClearOnExternalSourceLoadedDelegates();
diff --git a/Engine/Plugins/Enterprise/DatasmithImporter/Source/DirectLinkExtension/Private/DirectLinkManager.cpp b/Engine/Plugins/Enterprise/DatasmithImporter/Source/DirectLinkExtension/Private/DirectLinkManager.cpp
--- a/Engine/Plugins/Enterprise/DatasmithImporter/Source/DirectLinkExtension/Private/DirectLinkManager.cpp
+++ b/Engine/Plugins/Enterprise/DatasmithImporter/Source/DirectLinkExtension/Private/DirectLinkManager.cpp
@@ -3,6 +3,7 @@
 #include "DirectLinkManager.h"
 
 #include "DirectLinkAssetObserver.h"
+#include "DirectLinkExtensionSettings.h"
 #include "DirectLinkExternalSource.h"
 #include "DirectLinkUriResolver.h"
 
@@ -23,6 +24,57 @@ DEFINE_LOG_CATEGORY(LogDirectLinkManager);
 
 namespace UE::DatasmithImporter
 {
+	bool FDirectLinkAutoReconnectManager::Start()
+	{
+		if (GetDefault<UDirectLinkExtensionSettings>()->bAutoReconnect
+			&& (!CompletedFuture.IsValid() || CompletedFuture.IsReady()))
+		{
+			bShouldRun = true;
+			CompletedFuture = Async(EAsyncExecution::ThreadPool, [this]() {	Run(); });
+
+			return true;
+		}
+
+		return false;
+	}
+
+	void FDirectLinkAutoReconnectManager::Stop()
+	{
+		bShouldRun = false;
+	}
+
+	void FDirectLinkAutoReconnectManager::Run()
+	{
+		const float CurrentTime = FPlatformTime::Seconds();
+		const float TimeSinceLastTry = CurrentTime - LastTryTime;
+		const float ReconnectionDelayInSeconds = GetDefault<UDirectLinkExtensionSettings>()->ReconnectionDelayInSeconds;
+
+		if (TimeSinceLastTry < ReconnectionDelayInSeconds)
+		{
+			FPlatformProcess::Sleep(ReconnectionDelayInSeconds - TimeSinceLastTry);
+		}
+
+		int32 NumberOfSources;
+		{
+			FWriteScopeLock ReconnectionScopeLock(Manager.ReconnectionListLock);
+			for (int32 Index = Manager.ExternalSourcesToReconnect.Num() - 1; Index >= 0; --Index)
+			{
+				if (Manager.ExternalSourcesToReconnect[Index]->OpenStream())
+				{
+					Manager.ExternalSourcesToReconnect.RemoveAtSwap(Index);
+				}
+			}
+			NumberOfSources = Manager.ExternalSourcesToReconnect.Num();
+			LastTryTime = FPlatformTime::Seconds();
+		}
+
+		if (bShouldRun && NumberOfSources > 0)
+		{
+			// Could not reconnect, go back to the ThreadPool and try again later.
+			CompletedFuture = Async(EAsyncExecution::ThreadPool, [this]() {	Run(); });
+		}
+	}
+
 	/**
 	 * #ueent_todo: The AutoReimport feature should be generalize to all FExternalSource, not just DirectLink ones.
 	 */
@@ -44,6 +96,7 @@ namespace UE::DatasmithImporter
 	FDirectLinkManager::FDirectLinkManager()
 		: Endpoint(MakeUnique<DirectLink::FEndpoint>(TEXT("UE5-Editor")))
 		, AssetObserver(MakeUnique<FDirectLinkAssetObserver>(*this))
+		, ReconnectionManager(MakeUnique<FDirectLinkAutoReconnectManager>(*this))
 	{
 		Endpoint->AddEndpointObserver(this);
 
@@ -206,6 +259,52 @@ namespace UE::DatasmithImporter
 			}
 		}
 
+
+		TSet<DirectLink::FDestinationHandle> ActiveStreamsSources;
+		for (const DirectLink::FRawInfo::FStreamInfo& StreamInfo : RawInfoCache.StreamsInfo)
+		{
+			if (!StreamInfo.bIsActive)
+			{
+				continue;
+			}
+
+			if (TSharedRef<FDirectLinkExternalSource>* ExternalSource = DirectLinkSourceToExternalSourceMap.Find(StreamInfo.Source))
+			{
+				ActiveStreamsSources.Add(StreamInfo.Source);
+			}
+		}
+
+		{
+			FWriteScopeLock ReconnectionScopeLock(ReconnectionListLock);
+
+			for (int32 SourceIndex = ExternalSourcesToReconnect.Num() - 1; SourceIndex >= 0; --SourceIndex)
+			{
+				// If the source no longer exists, then there is no point in trying to reconnect.
+				if (InvalidExternalSourceIds.Contains(ExternalSourcesToReconnect[SourceIndex]->GetSourceHandle()))
+				{
+					ExternalSourcesToReconnect.RemoveAtSwap(SourceIndex);
+				}
+			}
+
+			for (const TPair<FGuid, TSharedRef<FDirectLinkExternalSource>>& ExternalSourceKeyValue : DirectLinkSourceToExternalSourceMap)
+			{
+				const TSharedRef<FDirectLinkExternalSource>& ExternalSource = ExternalSourceKeyValue.Value;
+
+				if (ExternalSource->IsStreamOpen() && !ActiveStreamsSources.Contains(ExternalSourceKeyValue.Key))
+				{
+					// Lost connection, update the external source state and try to reconnect.
+					ExternalSource->CloseStream();
+					
+					if (!ExternalSource->OpenStream())
+					{
+						// Could not reopen the stream, retry later.
+						ExternalSourcesToReconnect.Add(ExternalSource);
+						ReconnectionManager->Start();
+					}
+				}
+			}
+		}
+
 		// Remove all external sources that are no longer valid.
 		for (const DirectLink::FSourceHandle& SourceHandle : InvalidExternalSourceIds)
 		{
diff --git a/Engine/Plugins/Enterprise/DatasmithImporter/Source/DirectLinkExtension/Private/DirectLinkManager.h b/Engine/Plugins/Enterprise/DatasmithImporter/Source/DirectLinkExtension/Private/DirectLinkManager.h
--- a/Engine/Plugins/Enterprise/DatasmithImporter/Source/DirectLinkExtension/Private/DirectLinkManager.h
+++ b/Engine/Plugins/Enterprise/DatasmithImporter/Source/DirectLinkExtension/Private/DirectLinkManager.h
@@ -21,9 +21,40 @@ namespace UE::DatasmithImporter
 {
 	class FDirectLinkAssetObserver;
 	class FDirectLinkExternalSource;
+	class FDirectLinkManager;
 	struct FAutoReimportInfo;
 	struct FDirectLinkSourceDescription;
 
+	class FDirectLinkAutoReconnectManager
+	{
+	public:
+		FDirectLinkAutoReconnectManager(FDirectLinkManager& InManager)
+			: Manager(InManager)
+			, bShouldRun(true)
+		{}
+
+		~FDirectLinkAutoReconnectManager()
+		{
+			Stop();
+			CompletedFuture.Wait();
+		}
+
+		bool Start();
+
+		void Stop();
+
+	private:
+		void Run();
+
+		FDirectLinkManager& Manager;
+
+		TAtomic<bool> bShouldRun;
+
+		TFuture<void> CompletedFuture;
+
+		float LastTryTime = 0;
+	};
+
 	class FDirectLinkManager: public IDirectLinkManager, public DirectLink::IEndpointObserver
 	{
 	public:
@@ -113,6 +144,11 @@ namespace UE::DatasmithImporter
 		
 		TMap<DirectLink::FSourceHandle, TSharedRef<FDirectLinkExternalSource>> DirectLinkSourceToExternalSourceMap;
 
+		FRWLock ReconnectionListLock;
+		TArray<TSharedRef<FDirectLinkExternalSource>> ExternalSourcesToReconnect;
+
+		TUniquePtr<FDirectLinkAutoReconnectManager> ReconnectionManager;
+
 		TMap<UObject*, TSharedRef<FAutoReimportInfo>> RegisteredAutoReimportObjectMap;
 		
 		TMultiMap<TSharedRef<FExternalSource>, TSharedRef<FAutoReimportInfo>> RegisteredAutoReimportExternalSourceMap;
@@ -122,5 +158,6 @@ namespace UE::DatasmithImporter
 #if WITH_EDITOR
 		FDelegateHandle OnPIEEndHandle;
 #endif
+		friend FDirectLinkAutoReconnectManager;
 	};
 }
\ No newline at end of file
diff --git a/Engine/Plugins/Enterprise/DatasmithImporter/Source/DirectLinkExtension/Public/DirectLinkExtensionSettings.h b/Engine/Plugins/Enterprise/DatasmithImporter/Source/DirectLinkExtension/Public/DirectLinkExtensionSettings.h
new file mode 100644
--- /dev/null
+++ b/Engine/Plugins/Enterprise/DatasmithImporter/Source/DirectLinkExtension/Public/DirectLinkExtensionSettings.h
@@ -0,0 +1,28 @@
+// Copyright Epic Games, Inc. All Rights Reserved.
+
+#pragma once
+
+#include "CoreMinimal.h"
+#include "UObject/ObjectMacros.h"
+#include "Engine/EngineTypes.h"
+
+#include "DirectLinkExtensionSettings.generated.h"
+
+UCLASS(config=Editor)
+class DIRECTLINKEXTENSION_API UDirectLinkExtensionSettings : public UObject
+{
+	GENERATED_BODY()
+
+public:
+	/**
+	 * Whether or not to attempt reconnecting lost connection.
+	 */
+	UPROPERTY(config)
+	bool bAutoReconnect = true;
+
+	/**
+	 * Time in seconds to wait between reconnection attempts.
+	 */
+	UPROPERTY(config, meta = (ClampMin="0.1"))
+	float ReconnectionDelayInSeconds = 5.f;
+};
\ No newline at end of file
diff --git a/Engine/Plugins/Enterprise/DatasmithImporter/Source/DirectLinkExtension/Public/DirectLinkExternalSource.h b/Engine/Plugins/Enterprise/DatasmithImporter/Source/DirectLinkExtension/Public/DirectLinkExternalSource.h
--- a/Engine/Plugins/Enterprise/DatasmithImporter/Source/DirectLinkExtension/Public/DirectLinkExternalSource.h
+++ b/Engine/Plugins/Enterprise/DatasmithImporter/Source/DirectLinkExtension/Public/DirectLinkExternalSource.h
@@ -67,6 +67,10 @@ namespace UE::DatasmithImporter
 		 */
 		void Invalidate();
 
+		const FGuid& GetSourceHandle() const { return SourceHandle; }
+
+		const FGuid& GetDestinationHandle() const { return DestinationHandle; }
+
 	protected:
 		/**
 		 * Used by GetSceneReceiver() to create a wrapper around the ISceneReceiver returned by this function.
