FDirectLinkAssetObserver -  check if asset is loaded on update event, fixing a crash during Cooked Editor loading.

diff --git a/Engine/Plugins/Enterprise/DatasmithImporter/Source/DirectLinkExtension/Private/DirectLinkAssetObserver.cpp b/Engine/Plugins/Enterprise/DatasmithImporter/Source/DirectLinkExtension/Private/DirectLinkAssetObserver.cpp
--- a/Engine/Plugins/Enterprise/DatasmithImporter/Source/DirectLinkExtension/Private/DirectLinkAssetObserver.cpp
+++ b/Engine/Plugins/Enterprise/DatasmithImporter/Source/DirectLinkExtension/Private/DirectLinkAssetObserver.cpp
@@ -49,9 +49,12 @@ namespace UE::DatasmithImporter
 
 	void FDirectLinkAssetObserver::AssetUpdatedEvent(const FAssetData& AssetData)
 	{
-		if (UObject* UpdatedAsset = AssetData.GetAsset())
+		if (AssetData.IsAssetLoaded())
 		{
-			Manager.UpdateModifiedRegisteredAsset(UpdatedAsset);
+			if (UObject* UpdatedAsset = AssetData.GetAsset())
+			{
+				Manager.UpdateModifiedRegisteredAsset(UpdatedAsset);
+			}
 		}
 	}
 
