DirectLink in Editor - Fixing the source hash being invalid when reimporting without existing active connection.

JeanLuc.Corenthin

diff --git a/Engine/Plugins/Enterprise/DatasmithImporter/Source/DatasmithImporter/Private/DatasmithImportContext.cpp b/Engine/Plugins/Enterprise/DatasmithImporter/Source/DatasmithImporter/Private/DatasmithImportContext.cpp
--- a/Engine/Plugins/Enterprise/DatasmithImporter/Source/DatasmithImporter/Private/DatasmithImportContext.cpp
+++ b/Engine/Plugins/Enterprise/DatasmithImporter/Source/DatasmithImporter/Private/DatasmithImportContext.cpp
@@ -443,11 +443,23 @@ bool FDatasmithImportContext::SetupDestination(const FString& InImportPath, EObj
 
 void FDatasmithImportContext::InitScene(const TSharedRef<IDatasmithScene>& InScene)
 {
+	using namespace UE::DatasmithImporter;
 	Scene = InScene;
 
 	// Initialize the filtered scene as a copy of the original scene. We will use it to then filter out items to import.
 	FilteredScene = FDatasmithSceneFactory::DuplicateScene(Scene.ToSharedRef());
 	SceneName = FDatasmithUtils::SanitizeObjectName(Scene->GetName());
+
+	// If the hash was not valid during initialization, try to get it again, now that the scene is loaded.
+	if (!Options->SourceHash.IsValid())
+	{
+		const FSourceUri SourceUri(Options->SourceUri);
+
+		if (const TSharedPtr<FExternalSource> ExternalSource = IExternalSourceModule::GetOrCreateExternalSource(SourceUri))
+		{
+			Options->SourceHash = ExternalSource->GetSourceHash();
+		}
+	}
 }
 
 void FDatasmithImportContext::DisplayMessages()
