Datasmith Solidworks: Lights exported from solidwork should be converted to StdEmissive materials in UE

diff --git a/Engine/Plugins/Enterprise/DatasmithImporter/Source/DatasmithTranslator/Private/MasterMaterials/DatasmithStdMaterialSelector.cpp b/Engine/Plugins/Enterprise/DatasmithImporter/Source/DatasmithTranslator/Private/MasterMaterials/DatasmithStdMaterialSelector.cpp
--- a/Engine/Plugins/Enterprise/DatasmithImporter/Source/DatasmithTranslator/Private/MasterMaterials/DatasmithStdMaterialSelector.cpp
+++ b/Engine/Plugins/Enterprise/DatasmithImporter/Source/DatasmithTranslator/Private/MasterMaterials/DatasmithStdMaterialSelector.cpp
@@ -14,18 +14,26 @@ FDatasmithStdMaterialSelector::FDatasmithStdMaterialSelector()
 	// Master
 	MasterMaterialOpaque.FromSoftObjectPath( FSoftObjectPath("/DatasmithContent/Materials/StdOpaque/M_StdOpaque.M_StdOpaque") );
 	MasterMaterialTranslucent.FromSoftObjectPath( FSoftObjectPath("/DatasmithContent/Materials/StdTranslucent/M_StdTranslucent.M_StdTranslucent") );
+	MasterMaterialEmissive.FromSoftObjectPath( FSoftObjectPath("/DatasmithContent/Materials/StdEmissive/M_StdEmissive.M_StdEmissive") );
 }
 
 bool FDatasmithStdMaterialSelector::IsValid() const
 {
-	return MasterMaterialOpaque.IsValid() && MasterMaterialTranslucent.IsValid();
+	return MasterMaterialOpaque.IsValid() && MasterMaterialTranslucent.IsValid() && MasterMaterialEmissive.IsValid();
 }
 
 const FDatasmithMasterMaterial& FDatasmithStdMaterialSelector::GetMasterMaterial( const TSharedPtr< IDatasmithMasterMaterialElement >& InDatasmithMaterial ) const
 {
-	if (InDatasmithMaterial->GetMaterialType() == EDatasmithMasterMaterialType::Transparent)
+	switch (InDatasmithMaterial->GetMaterialType())
 	{
-		return MasterMaterialTranslucent;
+		case EDatasmithMasterMaterialType::Transparent:
+			return MasterMaterialTranslucent;
+			break;
+		case EDatasmithMasterMaterialType::Emissive:
+			return MasterMaterialEmissive;
+			break;
+		default:
+			return MasterMaterialOpaque;
+			break;
 	}
-	return MasterMaterialOpaque;
 }
diff --git a/Engine/Plugins/Enterprise/DatasmithImporter/Source/DatasmithTranslator/Private/MasterMaterials/DatasmithStdMaterialSelector.h b/Engine/Plugins/Enterprise/DatasmithImporter/Source/DatasmithTranslator/Private/MasterMaterials/DatasmithStdMaterialSelector.h
--- a/Engine/Plugins/Enterprise/DatasmithImporter/Source/DatasmithTranslator/Private/MasterMaterials/DatasmithStdMaterialSelector.h
+++ b/Engine/Plugins/Enterprise/DatasmithImporter/Source/DatasmithTranslator/Private/MasterMaterials/DatasmithStdMaterialSelector.h
@@ -22,4 +22,5 @@ public:
 private:
 	FDatasmithMasterMaterial MasterMaterialOpaque;
 	FDatasmithMasterMaterial MasterMaterialTranslucent;
+	FDatasmithMasterMaterial MasterMaterialEmissive;
 };
diff --git a/Engine/Source/Developer/Datasmith/DatasmithFacade/Public/DatasmithFacadeMaterial.h b/Engine/Source/Developer/Datasmith/DatasmithFacade/Public/DatasmithFacadeMaterial.h
--- a/Engine/Source/Developer/Datasmith/DatasmithFacade/Public/DatasmithFacadeMaterial.h
+++ b/Engine/Source/Developer/Datasmith/DatasmithFacade/Public/DatasmithFacadeMaterial.h
@@ -53,6 +53,7 @@ public:
 		Opaque,
 		Transparent,
 		ClearCoat,
+		Emissive,
 		/** Instantiate a master material from a specified one */
 		Custom,
 		/** Material has a transparent cutout map */
diff --git a/Engine/Source/Programs/Enterprise/Datasmith/DatasmithSolidworksExporter/Private/DatasmithExporter.cs b/Engine/Source/Programs/Enterprise/Datasmith/DatasmithSolidworksExporter/Private/DatasmithExporter.cs
--- a/Engine/Source/Programs/Enterprise/Datasmith/DatasmithSolidworksExporter/Private/DatasmithExporter.cs
+++ b/Engine/Source/Programs/Enterprise/Datasmith/DatasmithSolidworksExporter/Private/DatasmithExporter.cs
@@ -458,6 +458,7 @@ namespace DatasmithSolidworks
 
 				if (InMaterial.Emission > 0.0)
 				{
+					MasterMaterial.SetMaterialType(FDatasmithFacadeMasterMaterial.EMasterMaterialType.Emissive);
 					MasterMaterial.AddFloat("LuminanceAmount", (float)InMaterial.Emission);
 					MasterMaterial.AddColor("LuminanceFilter", R, G, B, 1.0f);
 				}
diff --git a/Engine/Source/Runtime/Datasmith/DatasmithCore/Public/DatasmithDefinitions.h b/Engine/Source/Runtime/Datasmith/DatasmithCore/Public/DatasmithDefinitions.h
--- a/Engine/Source/Runtime/Datasmith/DatasmithCore/Public/DatasmithDefinitions.h
+++ b/Engine/Source/Runtime/Datasmith/DatasmithCore/Public/DatasmithDefinitions.h
@@ -230,6 +230,7 @@ enum class EDatasmithMasterMaterialType : uint8
 	Opaque,
 	Transparent,
 	ClearCoat,
+	Emissive,
 	/** Instantiate a master material from a specified one */
 	Custom,
 	/** Material has a transparent cutout map */
