DirectLink in Editor - Auto-reimport is now driven by the DatasmithScene asset instead of the DatasmithSceneActor, this simplify the feature and avoids a lot of confusing edge-cases.

JeanLuc.Corenthin

diff --git a/Engine/Plugins/Enterprise/DatasmithContent/Source/DatasmithContentEditor/DatasmithContentEditor.Build.cs b/Engine/Plugins/Enterprise/DatasmithContent/Source/DatasmithContentEditor/DatasmithContentEditor.Build.cs
--- a/Engine/Plugins/Enterprise/DatasmithContent/Source/DatasmithContentEditor/DatasmithContentEditor.Build.cs
+++ b/Engine/Plugins/Enterprise/DatasmithContent/Source/DatasmithContentEditor/DatasmithContentEditor.Build.cs
@@ -29,6 +29,7 @@ namespace UnrealBuildTool.Rules
 					"PropertyEditor",
 					"SlateCore",
 					"Slate",
+					"ToolMenus",
 					"DetailCustomizations",
                 }
 			);
diff --git a/Engine/Plugins/Enterprise/DatasmithContent/Source/DatasmithContentEditor/Private/AssetTypeActions_DatasmithScene.cpp b/Engine/Plugins/Enterprise/DatasmithContent/Source/DatasmithContentEditor/Private/AssetTypeActions_DatasmithScene.cpp
--- a/Engine/Plugins/Enterprise/DatasmithContent/Source/DatasmithContentEditor/Private/AssetTypeActions_DatasmithScene.cpp
+++ b/Engine/Plugins/Enterprise/DatasmithContent/Source/DatasmithContentEditor/Private/AssetTypeActions_DatasmithScene.cpp
@@ -7,6 +7,7 @@
 #include "DatasmithScene.h"
 
 #include "Framework/MultiBox/MultiBoxBuilder.h"
+#include "ToolMenuSection.h"
 
 #define LOCTEXT_NAMESPACE "AssetTypeActions_DatasmithScene"
 
@@ -25,6 +26,47 @@ UClass* FAssetTypeActions_DatasmithScene::GetSupportedClass() const
 	return UDatasmithScene::StaticClass();
 }
 
+void FAssetTypeActions_DatasmithScene::GetActions(const TArray<UObject*>& InObjects, FToolMenuSection& Section)
+{
+	TArray<TWeakObjectPtr<UDatasmithScene>> Scenes = GetTypedWeakObjectPtrs<UDatasmithScene>(InObjects);
+
+	FilterByDirectLinkAutoReimportSupport(Scenes);
+
+	const bool bHasCanExecuteAutoReimport = Scenes.Num() > 0;
+	
+	bool bHasAutoReimportEnabled = false;
+	const IDatasmithContentEditorModule& DatasmithContentEditorModule = IDatasmithContentEditorModule::Get();
+	for (const TWeakObjectPtr<UDatasmithScene>& Scene : Scenes)
+	{
+		if (Scene.IsValid())
+		{
+			bHasAutoReimportEnabled = DatasmithContentEditorModule.IsAssetAutoReimportEnabled(Scene.Get()).Get(false);
+			if (bHasAutoReimportEnabled)
+			{
+				break;
+			}
+		}
+	}
+
+	const FText AutoReimportText = bHasAutoReimportEnabled
+		? LOCTEXT("DatasmithScene_ToggleDirectLinkAutoReimport_Disable", "Disable Direct Link auto-reimport")
+		: LOCTEXT("DatasmithScene_ToggleDirectLinkAutoReimport_Enable", "Enable Direct Link auto-reimport");
+	const FText AutoReimportTooltip = bHasAutoReimportEnabled
+		? LOCTEXT("DatasmithScene_ToggleDirectLinkAutoReimportTooltip_Disable", "Disable Direct Link auto-reimport for all selected Datasmith Scenes.")
+		: LOCTEXT("DatasmithScene_ToggleDirectLinkAutoReimportTooltip_Enable", "Enable Direct Link auto-reimport for all selected Datasmith Scenes.");
+
+	Section.AddMenuEntry(
+		"DatasmithScene_ToggleDirectLinkAutoReimport",
+		AutoReimportText,
+		AutoReimportTooltip,
+		FSlateIcon(FAppStyle::GetAppStyleSetName(), "ContentBrowser.AssetActions.ReimportWithNewFile"),
+		FUIAction(
+			FExecuteAction::CreateLambda([Scenes = MoveTemp(Scenes), bHasAutoReimportEnabled]() { FAssetTypeActions_DatasmithScene::ExecuteToggleDirectLinkAutoReimport(Scenes, !bHasAutoReimportEnabled); }),
+			FCanExecuteAction::CreateLambda([bHasCanExecuteAutoReimport]() { return bHasCanExecuteAutoReimport; })
+		)
+	);
+}
+
 void FAssetTypeActions_DatasmithScene::GetResolvedSourceFilePaths(const TArray<UObject*>& TypeAssets, TArray<FString>& OutSourceFilePaths) const
 {
 	for ( const UObject* Asset : TypeAssets )
@@ -63,4 +105,36 @@ void FAssetTypeActions_DatasmithScene::OpenAssetEditor(const TArray<UObject*>& I
 	}
 }
 
+void FAssetTypeActions_DatasmithScene::ExecuteToggleDirectLinkAutoReimport(const TArray<TWeakObjectPtr<UDatasmithScene>>& Scenes, bool bEnabled)
+{
+	const IDatasmithContentEditorModule& DatasmithContentEditorModule = IDatasmithContentEditorModule::Get();
+
+	for (const TWeakObjectPtr<class UDatasmithScene>& Scene : Scenes)
+	{
+		if (Scene.IsValid())
+		{
+			DatasmithContentEditorModule.SetAssetAutoReimport(Scene.Get(), bEnabled);
+		}
+	}
+}
+
+void FAssetTypeActions_DatasmithScene::FilterByDirectLinkAutoReimportSupport(TArray<TWeakObjectPtr<UDatasmithScene>>& Scenes)
+{
+	const IDatasmithContentEditorModule& DatasmithContentEditorModule = IDatasmithContentEditorModule::Get();
+
+	for (int32 SceneIndex = Scenes.Num() - 1; SceneIndex >= 0; --SceneIndex)
+	{
+		if (Scenes[SceneIndex].IsValid())
+		{
+			const bool bIsReimportAvailable = DatasmithContentEditorModule.IsAssetAutoReimportAvailable(Scenes[SceneIndex].Get()).Get(false);
+			if (bIsReimportAvailable)
+			{
+				continue;
+			}
+		}
+		Scenes.RemoveAtSwap(SceneIndex);
+	}
+}
+
+
 #undef LOCTEXT_NAMESPACE
diff --git a/Engine/Plugins/Enterprise/DatasmithContent/Source/DatasmithContentEditor/Private/AssetTypeActions_DatasmithScene.h b/Engine/Plugins/Enterprise/DatasmithContent/Source/DatasmithContentEditor/Private/AssetTypeActions_DatasmithScene.h
--- a/Engine/Plugins/Enterprise/DatasmithContent/Source/DatasmithContentEditor/Private/AssetTypeActions_DatasmithScene.h
+++ b/Engine/Plugins/Enterprise/DatasmithContent/Source/DatasmithContentEditor/Private/AssetTypeActions_DatasmithScene.h
@@ -14,9 +14,16 @@ public:
 	virtual uint32 GetCategories() override;
 	virtual FColor GetTypeColor() const override { return FColor(255, 0, 0); }
 	virtual UClass* GetSupportedClass() const override;
-	virtual bool HasActions ( const TArray<UObject*>& InObjects ) const override { return false; }
+	virtual void GetActions(const TArray<UObject*>& InObjects, FToolMenuSection& Section) override;
+	virtual bool HasActions ( const TArray<UObject*>& InObjects ) const override { return true; }
 	virtual bool IsImportedAsset() const override { return true; }
 	virtual void GetResolvedSourceFilePaths(const TArray<UObject*>& TypeAssets, TArray<FString>& OutSourceFilePaths) const override;
 	virtual void OpenAssetEditor(const TArray<UObject*>& InObjects, TSharedPtr<class IToolkitHost> EditWithinLevelEditor = TSharedPtr<IToolkitHost>()) override;
 	// End IAssetTypeActions interface
+
+private:
+
+	static void ExecuteToggleDirectLinkAutoReimport(const TArray<TWeakObjectPtr<class UDatasmithScene>>& Scenes, bool bEnabled);
+
+	static void FilterByDirectLinkAutoReimportSupport(TArray<TWeakObjectPtr<class UDatasmithScene>>& Scenes);
 };
diff --git a/Engine/Plugins/Enterprise/DatasmithContent/Source/DatasmithContentEditor/Private/DatasmithContentEditorModule.cpp b/Engine/Plugins/Enterprise/DatasmithContent/Source/DatasmithContentEditor/Private/DatasmithContentEditorModule.cpp
--- a/Engine/Plugins/Enterprise/DatasmithContent/Source/DatasmithContentEditor/Private/DatasmithContentEditorModule.cpp
+++ b/Engine/Plugins/Enterprise/DatasmithContent/Source/DatasmithContentEditor/Private/DatasmithContentEditorModule.cpp
@@ -44,8 +44,6 @@ public:
 		FDatasmithContentEditorStyle::Initialize();
 
 		RegisterDetailCustomization();
-
-		OnMapChangeHandle = FEditorDelegates::MapChange.AddRaw(this, &FDatasmithContentEditorModule::OnMapChange);
 	}
 
 	virtual void ShutdownModule() override
@@ -73,9 +71,6 @@ public:
 		FDatasmithContentEditorStyle::Shutdown();
 
 		UnregisterDetailCustomization();
-
-		ClearAutoReimportAssets();
-		FEditorDelegates::MapChange.Remove(OnMapChangeHandle);
 	}
 
 	virtual void RegisterSpawnDatasmithSceneActorsHandler( FOnSpawnDatasmithSceneActors InSpawnActorsDelegate ) override
@@ -146,17 +141,8 @@ public:
 		}
 	}
 		
-	virtual TOptional<bool> SetAssetAutoReimport(UObject* Asset, bool bEnabled) override
+	virtual TOptional<bool> SetAssetAutoReimport(UObject* Asset, bool bEnabled) const override
 	{
-		if (bEnabled)
-		{
-			AutoReimportingAssets.Add(Asset);
-		}
-		else
-		{
-			AutoReimportingAssets.Remove(Asset);
-		}
-
 		if (SetAssetAutoReimportHandler.IsBound())
 		{
 			return SetAssetAutoReimportHandler.Execute(Asset, bEnabled);
@@ -263,28 +249,6 @@ private:
 		return TSharedPtr<IDataprepImporterInterface>();
 	}
 
-	void OnMapChange(uint32 MapEventFlags)
-	{
-		// We must clear all auto-reimport registered assets before changing map.
-		ClearAutoReimportAssets();
-	}
-
-	void ClearAutoReimportAssets()
-	{
-		if (SetAssetAutoReimportHandler.IsBound())
-		{
-			for (const TSoftObjectPtr<UObject>& Asset : AutoReimportingAssets)
-			{
-				if (Asset.IsValid())
-				{
-					SetAssetAutoReimportHandler.Execute(Asset.Get(), false);
-				}
-			}
-		}
-
-		AutoReimportingAssets.Empty();
-	}
-
 	void RegisterDetailCustomization()
 	{
 		const FName PropertyEditor("PropertyEditor");
@@ -315,9 +279,6 @@ private:
 	FOnIsAssetAutoReimportEnabled IsAssetAutoReimportEnabledHandler;
 	FOnBrowseExternalSourceUri BrowseExternalSourceUriHandler;
 	FOnGetSupportedUriSchemes GetSupportedUriSchemeHandler;
-
-	FDelegateHandle OnMapChangeHandle;
-	TSet<TSoftObjectPtr<UObject>> AutoReimportingAssets;
 };
 
 #undef LOCTEXT_NAMESPACE
diff --git a/Engine/Plugins/Enterprise/DatasmithContent/Source/DatasmithContentEditor/Private/DatasmithSceneActorDetailsPanel.cpp b/Engine/Plugins/Enterprise/DatasmithContent/Source/DatasmithContentEditor/Private/DatasmithSceneActorDetailsPanel.cpp
--- a/Engine/Plugins/Enterprise/DatasmithContent/Source/DatasmithContentEditor/Private/DatasmithSceneActorDetailsPanel.cpp
+++ b/Engine/Plugins/Enterprise/DatasmithContent/Source/DatasmithContentEditor/Private/DatasmithSceneActorDetailsPanel.cpp
@@ -90,23 +90,6 @@ void FDatasmithSceneActorDetailsPanel::CustomizeDetails(IDetailLayoutBuilder& De
 		[
 			WrapBox
 		];
-
-	ActionsCategory.AddCustomRow(AutoReimportCaption)
-		.NameContent()
-		[
-			SNew(STextBlock)
-			.Font(FEditorStyle::GetFontStyle(TEXT("PropertyWindow.NormalFont")))
-			.Text(AutoReimportCaption)
-			.ToolTipText(AutoReimportTooltip)
-			.IsEnabled(this, &FDatasmithSceneActorDetailsPanel::GetAutoReimportIsEnabled)
-		]
-		.ValueContent()
-		[
-			SNew(SCheckBox)
-			.IsChecked(this, &FDatasmithSceneActorDetailsPanel::GetAutoReimportIsChecked)
-			.IsEnabled(this, &FDatasmithSceneActorDetailsPanel::GetAutoReimportIsEnabled)
-			.OnCheckStateChanged(FOnCheckStateChanged::CreateSP(this, &FDatasmithSceneActorDetailsPanel::OnAutoReimportStateChanged))
-		];
 }
 
 FReply FDatasmithSceneActorDetailsPanel::OnExecuteAction()
@@ -122,66 +105,4 @@ FReply FDatasmithSceneActorDetailsPanel::OnExecuteAction()
 	return FReply::Handled();
 }
 
-void FDatasmithSceneActorDetailsPanel::OnAutoReimportStateChanged(ECheckBoxState NewState)
-{
-	IDatasmithContentEditorModule& DatasmithContentEditorModule = IDatasmithContentEditorModule::Get();
-
-	const bool bEnabled = NewState == ECheckBoxState::Checked;
-	for (const TWeakObjectPtr< UObject >& SelectedObject : SelectedObjectsList)
-	{
-		if (ADatasmithSceneActor* SceneActor = Cast< ADatasmithSceneActor >(SelectedObject.Get()))
-		{
-			DatasmithContentEditorModule.SetAssetAutoReimport(SceneActor->Scene.Get(), bEnabled);
-		}
-	}
-}
-
-ECheckBoxState FDatasmithSceneActorDetailsPanel::GetAutoReimportIsChecked() const
-{
-	IDatasmithContentEditorModule& DatasmithContentEditorModule = IDatasmithContentEditorModule::Get();
-	ECheckBoxState State = ECheckBoxState::Unchecked;
-	bool bHasCheckedAssets = false;
-
-	for (const TWeakObjectPtr< UObject >& SelectedObject : SelectedObjectsList)
-	{
-		if (ADatasmithSceneActor* SceneActor = Cast< ADatasmithSceneActor >(SelectedObject.Get()))
-		{
-			const TOptional<bool> bIsAutoReimportEnabled = DatasmithContentEditorModule.IsAssetAutoReimportEnabled(SceneActor->Scene.Get());
-			const bool bDefaultValue = false;
-
-			if (bIsAutoReimportEnabled.Get(bDefaultValue))
-			{
-				bHasCheckedAssets = true;
-			}
-			else if (bHasCheckedAssets)
-			{
-				return ECheckBoxState::Undetermined;
-			}
-		}
-	}
-
-	return bHasCheckedAssets ? ECheckBoxState::Checked : ECheckBoxState::Unchecked;
-}
-
-bool FDatasmithSceneActorDetailsPanel::GetAutoReimportIsEnabled() const
-{
-	IDatasmithContentEditorModule& DatasmithContentEditorModule = IDatasmithContentEditorModule::Get();
-
-	for (const TWeakObjectPtr< UObject >& SelectedObject : SelectedObjectsList)
-	{
-		if (ADatasmithSceneActor* SceneActor = Cast< ADatasmithSceneActor >(SelectedObject.Get()))
-		{
-			const TOptional<bool> bIsAutoReimportAvailable = DatasmithContentEditorModule.IsAssetAutoReimportAvailable(SceneActor->Scene.Get());
-			const bool bDefaultValue = false;
-
-			if (!bIsAutoReimportAvailable.Get(bDefaultValue))
-			{
-				return false;
-			}
-		}
-	}
-
-	return true;
-};
-
 #undef LOCTEXT_NAMESPACE
diff --git a/Engine/Plugins/Enterprise/DatasmithContent/Source/DatasmithContentEditor/Private/DatasmithSceneActorDetailsPanel.h b/Engine/Plugins/Enterprise/DatasmithContent/Source/DatasmithContentEditor/Private/DatasmithSceneActorDetailsPanel.h
--- a/Engine/Plugins/Enterprise/DatasmithContent/Source/DatasmithContentEditor/Private/DatasmithSceneActorDetailsPanel.h
+++ b/Engine/Plugins/Enterprise/DatasmithContent/Source/DatasmithContentEditor/Private/DatasmithSceneActorDetailsPanel.h
@@ -24,12 +24,6 @@ public:
 protected:
 	FReply OnExecuteAction();
 
-	void OnAutoReimportStateChanged(ECheckBoxState NewState);
-
-	ECheckBoxState GetAutoReimportIsChecked() const;
-
-	bool GetAutoReimportIsEnabled() const;
-
 private:
 	TArray< TWeakObjectPtr< UObject > > SelectedObjectsList;
 	bool bReimportDeletedActors;
diff --git a/Engine/Plugins/Enterprise/DatasmithContent/Source/DatasmithContentEditor/Public/DatasmithContentEditorModule.h b/Engine/Plugins/Enterprise/DatasmithContent/Source/DatasmithContentEditor/Public/DatasmithContentEditorModule.h
--- a/Engine/Plugins/Enterprise/DatasmithContent/Source/DatasmithContentEditor/Public/DatasmithContentEditorModule.h
+++ b/Engine/Plugins/Enterprise/DatasmithContent/Source/DatasmithContentEditor/Public/DatasmithContentEditorModule.h
@@ -113,7 +113,7 @@ public:
 	 */
 	virtual void RegisterSetAssetAutoReimportHandler(FOnSetAssetAutoReimport&& SetAssetAutoReimportDelegate) = 0;
 	virtual void UnregisterSetAssetAutoReimportHandler(FDelegateHandle InHandle) = 0;
-	virtual TOptional<bool> SetAssetAutoReimport(UObject* Asset, bool bEnabled) = 0;
+	virtual TOptional<bool> SetAssetAutoReimport(UObject* Asset, bool bEnabled) const = 0;
 
 	/**
 	 * Delegate returning if the AutoReimport feature is available for a given asset.
diff --git a/Engine/Plugins/Enterprise/DatasmithImporter/Source/DatasmithImporter/Private/DatasmithImportFactory.cpp b/Engine/Plugins/Enterprise/DatasmithImporter/Source/DatasmithImporter/Private/DatasmithImportFactory.cpp
--- a/Engine/Plugins/Enterprise/DatasmithImporter/Source/DatasmithImporter/Private/DatasmithImportFactory.cpp
+++ b/Engine/Plugins/Enterprise/DatasmithImporter/Source/DatasmithImporter/Private/DatasmithImportFactory.cpp
@@ -553,13 +553,24 @@ bool UDatasmithImportFactory::Import( FDatasmithImportContext& ImportContext )
 	}
 
 	bool bOutOperationCanceled = false;
-	bool bImportSuccess = DatasmithImportFactoryImpl::ImportDatasmithScene( ImportContext, bOutOperationCanceled );
+	const bool bImportSuccess = DatasmithImportFactoryImpl::ImportDatasmithScene( ImportContext, bOutOperationCanceled );
 
 	GEditor->RedrawAllViewports();
 
 	if ( !IsAutomatedImport() )
 	{
 		ImportContext.DisplayMessages();
+
+		if ( bImportSuccess && !bOutOperationCanceled )
+		{
+			IDatasmithContentEditorModule& ContentEditorModule = IDatasmithContentEditorModule::Get();
+			const bool bIsAutoReimportAvailable = ContentEditorModule.IsAssetAutoReimportAvailable( ImportContext.SceneAsset ).Get( false );
+
+			if ( bIsAutoReimportAvailable )
+			{
+				ContentEditorModule.SetAssetAutoReimport( ImportContext.SceneAsset, true );
+			}
+		}
 	}
 
 	if ( bOutOperationCanceled )
diff --git a/Engine/Plugins/Enterprise/DatasmithImporter/Source/DatasmithImporter/Private/DatasmithImporterModule.cpp b/Engine/Plugins/Enterprise/DatasmithImporter/Source/DatasmithImporter/Private/DatasmithImporterModule.cpp
--- a/Engine/Plugins/Enterprise/DatasmithImporter/Source/DatasmithImporter/Private/DatasmithImporterModule.cpp
+++ b/Engine/Plugins/Enterprise/DatasmithImporter/Source/DatasmithImporter/Private/DatasmithImporterModule.cpp
@@ -765,13 +765,7 @@ void FDatasmithImporterModule::SetupDatasmithContentDelegates()
 				const FAssetData AssetData(Asset);
 				const FSourceUri SourceUri = FSourceUri::FromAssetData(AssetData);
 
-				if (!SourceUri.HasScheme(FDirectLinkUriResolver::GetDirectLinkScheme()))
-				{
-					return false;
-				}
-
-				TSharedPtr<FExternalSource> ExternalSource = IExternalSourceModule::GetOrCreateExternalSource(SourceUri);
-				return ExternalSource && ExternalSource->IsAvailable();
+				return SourceUri.HasScheme(FDirectLinkUriResolver::GetDirectLinkScheme());
 			});
 		IsAssetAutoReimportSupportedHandle = IsAssetAutoReimportAvailable.GetHandle();
 		DatasmithContentEditorModule.RegisterIsAssetAutoReimportAvailableHandler(MoveTemp(IsAssetAutoReimportAvailable));
@@ -785,12 +779,12 @@ void FDatasmithImporterModule::SetupDatasmithContentDelegates()
 				const FAssetData AssetData(Asset);
 				const FSourceUri SourceUri = FSourceUri::FromAssetData(AssetData);
 
-				if (!SourceUri.HasScheme(FDirectLinkUriResolver::GetDirectLinkScheme()))
+				if (SourceUri.HasScheme(FDirectLinkUriResolver::GetDirectLinkScheme()))
 				{
-					return false;
+					return IDirectLinkExtensionModule::Get().GetManager().IsAssetAutoReimportEnabled(Asset);
 				}
 
-				return IDirectLinkExtensionModule::Get().GetManager().IsAssetAutoReimportEnabled(Asset);
+				return false;
 			});
 		IsAssetAutoReimportEnabledHandle = IsAssetAutoReimportEnabled.GetHandle();
 		DatasmithContentEditorModule.RegisterIsAssetAutoReimportEnabledHandler(MoveTemp(IsAssetAutoReimportEnabled));
diff --git a/Engine/Plugins/Enterprise/DatasmithImporter/Source/DirectLinkExtension/Private/DirectLinkManager.cpp b/Engine/Plugins/Enterprise/DatasmithImporter/Source/DirectLinkExtension/Private/DirectLinkManager.cpp
--- a/Engine/Plugins/Enterprise/DatasmithImporter/Source/DirectLinkExtension/Private/DirectLinkManager.cpp
+++ b/Engine/Plugins/Enterprise/DatasmithImporter/Source/DirectLinkExtension/Private/DirectLinkManager.cpp
@@ -63,7 +63,7 @@ namespace UE::DatasmithImporter
 			FPlatformProcess::Sleep(ReconnectionDelayInSeconds - TimeSinceLastTry);
 		}
 
-		int32 NumberOfSources;
+		bool bHasSourceToReconnect = false;
 		{
 			FWriteScopeLock ReconnectionScopeLock(Manager.ReconnectionListLock);
 			for (int32 Index = Manager.ExternalSourcesToReconnect.Num() - 1; Index >= 0; --Index)
@@ -73,11 +73,32 @@ namespace UE::DatasmithImporter
 					Manager.ExternalSourcesToReconnect.RemoveAtSwap(Index);
 				}
 			}
-			NumberOfSources = Manager.ExternalSourcesToReconnect.Num();
-			LastTryTime = FPlatformTime::Seconds();
+			bHasSourceToReconnect = Manager.ExternalSourcesToReconnect.Num() > 0;
 		}
 
-		if (bShouldRun && NumberOfSources > 0)
+		TArray<TSoftObjectPtr<UObject>> AutoReimportObjectListCopy;
+		{
+			FRWScopeLock ScopeLock(Manager.AutoReimportLock, FRWScopeLockType::SLT_ReadOnly);
+			AutoReimportObjectListCopy = Manager.PendingAutoReimportObjects.Array();
+		}
+
+		for (const TSoftObjectPtr<UObject>& Object : AutoReimportObjectListCopy)
+		{
+			if (Object.IsValid())
+			{
+				Manager.EnableAssetAutoReimport(Object.Get());
+			}
+		}
+
+		if (AutoReimportObjectListCopy.Num() > 0 && !bHasSourceToReconnect)
+		{
+			FRWScopeLock ScopeLock(Manager.AutoReimportLock, FRWScopeLockType::SLT_ReadOnly);
+			bHasSourceToReconnect = Manager.PendingAutoReimportObjects.Num() > 0;
+		}
+
+		LastTryTime = FPlatformTime::Seconds();
+
+		if (bShouldRun && bHasSourceToReconnect)
 		{
 			// Could not reconnect, go back to the ThreadPool and try again later.
 			CompletedFuture = Async(EAsyncExecution::ThreadPool, [this]() {	Run(); });
@@ -119,7 +140,7 @@ namespace UE::DatasmithImporter
 		Endpoint->RemoveEndpointObserver(this);
 
 		// Make sure all DirectLink external source become stales and their delegates stripped.
-		for (const TPair<FSourceUri, TSharedRef<FDirectLinkExternalSource>>& UriExternalSourcePair : UriToExternalSourceMap)
+		for (const TPair<DirectLink::FSourceHandle, TSharedRef<FDirectLinkExternalSource>>& UriExternalSourcePair : DirectLinkSourceToExternalSourceMap)
 		{
 			UriExternalSourcePair.Value->Invalidate();
 		}
@@ -131,12 +152,10 @@ namespace UE::DatasmithImporter
 
 	TSharedPtr<FDirectLinkExternalSource> FDirectLinkManager::GetOrCreateExternalSource(const DirectLink::FSourceHandle& SourceHandle)
 	{
-		TSharedPtr<FDirectLinkExternalSource> DirectLinkExternalSource = nullptr;
-
 		if (TSharedRef<FDirectLinkExternalSource>* ExternalSourceEntry = DirectLinkSourceToExternalSourceMap.Find(SourceHandle))
 		{
 			// A DirectLinkExternalSource already exists for this SourceHandle.
-			DirectLinkExternalSource = *ExternalSourceEntry;
+			return *ExternalSourceEntry;
 		}
 		else if (RegisteredExternalSourcesInfo.Num() > 0)
 		{
@@ -152,23 +171,22 @@ namespace UE::DatasmithImporter
 
 				for (const FDirectLinkExternalSourceRegisterInformation& RegisteredInfo : RegisteredExternalSourcesInfo)
 				{
-					TSharedPtr<FDirectLinkExternalSource> UninitializedExternalSource = RegisteredInfo.SpawnFunction(ExternalSourceUri);
+					TSharedRef<FDirectLinkExternalSource> SpawnedExternalSource = RegisteredInfo.SpawnFunction(ExternalSourceUri);
 
-					if (UninitializedExternalSource->CanOpenNewConnection(SourceInfo))
+					if (SpawnedExternalSource->CanOpenNewConnection(SourceInfo))
 					{
-						DirectLinkExternalSource = UninitializedExternalSource;
-						const FGuid DestinationHandle = Endpoint->AddDestination(ExternalSourceName, DirectLink::EVisibility::Private, DirectLinkExternalSource);
-						DirectLinkExternalSource->Initialize(SourceName, SourceHandle, DestinationHandle);
+						const FGuid DestinationHandle = Endpoint->AddDestination(ExternalSourceName, DirectLink::EVisibility::Private, SpawnedExternalSource);
+						SpawnedExternalSource->Initialize(SourceName, SourceHandle, DestinationHandle);
 
-						DirectLinkSourceToExternalSourceMap.Add(SourceHandle, DirectLinkExternalSource.ToSharedRef());
-						UriToExternalSourceMap.Add(ExternalSourceUri, DirectLinkExternalSource.ToSharedRef());
-						break;
+						DirectLinkSourceToExternalSourceMap.Add(SourceHandle, SpawnedExternalSource);
+						
+						return SpawnedExternalSource;
 					}
 				}
 			}
 		}
 
-		return DirectLinkExternalSource;
+		return nullptr;
 	}
 
 	TSharedPtr<FDirectLinkExternalSource> FDirectLinkManager::GetOrCreateExternalSource(const FSourceUri& Uri)
@@ -375,15 +393,28 @@ namespace UE::DatasmithImporter
 	void FDirectLinkManager::InvalidateSource(const DirectLink::FSourceHandle& InvalidSourceHandle)
 	{
 		TSharedRef<FDirectLinkExternalSource> DirectLinkExternalSource = DirectLinkSourceToExternalSourceMap.FindAndRemoveChecked(InvalidSourceHandle);
-		UriToExternalSourceMap.Remove(DirectLinkExternalSource->GetSourceUri());
 
 		// Clear the auto-reimport cache for this external source.
-		TArray<TSharedRef<FAutoReimportInfo>> AutoReimportInfoList;
-		RegisteredAutoReimportExternalSourceMap.MultiFind(DirectLinkExternalSource, AutoReimportInfoList);
-		RegisteredAutoReimportExternalSourceMap.Remove(DirectLinkExternalSource);
-		for (const TSharedRef<FAutoReimportInfo>& AutoReimportInfo : AutoReimportInfoList)
 		{
-			RegisteredAutoReimportObjectMap.Remove(AutoReimportInfo->TargetObject.Get());
+			FRWScopeLock ScopeLock(AutoReimportLock, FRWScopeLockType::SLT_Write);
+			TArray<TSharedRef<FAutoReimportInfo>> AutoReimportInfoList;
+			AutoReimportExternalSourcesMap.MultiFind(DirectLinkExternalSource, AutoReimportInfoList);
+			AutoReimportExternalSourcesMap.Remove(DirectLinkExternalSource);
+
+			for (const TSharedRef<FAutoReimportInfo>& AutoReimportInfo : AutoReimportInfoList)
+			{
+				UObject* Asset = AutoReimportInfo->TargetObject.Get();
+
+				// Assets that were registered for auto-reimport are unregistered from the source
+				// and put back in the "PendingRegistration" list. That way if a compatible source appear they will register to it.
+				AutoReimportObjectsMap.Remove(Asset);
+				PendingAutoReimportObjects.Add(Asset);
+			}
+
+			if (PendingAutoReimportObjects.Num() > 0)
+			{
+				ReconnectionManager->Start();
+			}
 		}
 
 		DirectLinkExternalSource->Invalidate();
@@ -391,7 +422,9 @@ namespace UE::DatasmithImporter
 
 	bool FDirectLinkManager::IsAssetAutoReimportEnabled(UObject* InAsset) const
 	{
-		return RegisteredAutoReimportObjectMap.Find(InAsset) != nullptr;
+		FRWScopeLock ScopeLock(AutoReimportLock, FRWScopeLockType::SLT_ReadOnly);
+
+		return AutoReimportObjectsMap.Contains(InAsset)	|| PendingAutoReimportObjects.Contains(InAsset);
 	}
 
 	bool FDirectLinkManager::SetAssetAutoReimport(UObject* InAsset, bool bEnabled)
@@ -406,7 +439,8 @@ namespace UE::DatasmithImporter
 		const FSourceUri Uri = FSourceUri::FromAssetData(AssetData);
 		const bool bIsValidDirectLinkUri = Uri.IsValid() && Uri.HasScheme(FDirectLinkUriResolver::GetDirectLinkScheme());
 
-		if (bIsValidDirectLinkUri && !RegisteredAutoReimportObjectMap.Contains(InAsset))
+		FRWScopeLock ScopeLock(AutoReimportLock, FRWScopeLockType::SLT_Write);
+		if (bIsValidDirectLinkUri && !AutoReimportObjectsMap.Contains(InAsset))
 		{
 			if (TSharedPtr<FDirectLinkExternalSource> ExternalSource = GetOrCreateExternalSource(Uri))
 			{
@@ -417,12 +451,22 @@ namespace UE::DatasmithImporter
 				TSharedRef<FDirectLinkExternalSource> ExternalSourceRef = ExternalSource.ToSharedRef();
 				TSharedRef<FAutoReimportInfo> AutoReimportInfo = MakeShared<FAutoReimportInfo>(InAsset, ExternalSourceRef, DelegateHandle);
 
-				RegisteredAutoReimportObjectMap.Add(InAsset, AutoReimportInfo);
-				RegisteredAutoReimportExternalSourceMap.Add(ExternalSourceRef, AutoReimportInfo);
+				AutoReimportObjectsMap.Add(InAsset, AutoReimportInfo);
+				AutoReimportExternalSourcesMap.Add(ExternalSourceRef, AutoReimportInfo);
 				ExternalSource->OpenStream();
-
+				PendingAutoReimportObjects.Remove(InAsset);
 				return true;
 			}
+			else
+			{
+				bool bIsAlreadySet = false;
+				PendingAutoReimportObjects.Add(InAsset, &bIsAlreadySet);
+				if (!bIsAlreadySet)
+				{
+					ReconnectionManager->Start();
+					return true;
+				}
+			}
 		}
 
 		return false;
@@ -431,14 +475,20 @@ namespace UE::DatasmithImporter
 	bool FDirectLinkManager::DisableAssetAutoReimport(UObject* InAsset)
 	{
 		// Disable auto-reimport for InAsset.
-		if (const TSharedRef<FAutoReimportInfo>* AutoReimportInfoPtr = RegisteredAutoReimportObjectMap.Find(InAsset))
+		FRWScopeLock ScopeLock(AutoReimportLock, FRWScopeLockType::SLT_Write);
+		if (const TSharedRef<FAutoReimportInfo>* AutoReimportInfoPtr = AutoReimportObjectsMap.Find(InAsset))
 		{
 			// Holding a local reference to the FAutoReimportInfo to ensure its lifetime while we are cleaning up.
 			const TSharedRef<FAutoReimportInfo> AutoReimportInfo(*AutoReimportInfoPtr);
 
 			AutoReimportInfo->ExternalSource->OnExternalSourceChanged.Remove(AutoReimportInfo->ImportDelegateHandle);
-			RegisteredAutoReimportObjectMap.Remove(InAsset);
-			RegisteredAutoReimportExternalSourceMap.RemoveSingle(AutoReimportInfo->ExternalSource, AutoReimportInfo);
+			AutoReimportObjectsMap.Remove(InAsset);
+			AutoReimportExternalSourcesMap.RemoveSingle(AutoReimportInfo->ExternalSource, AutoReimportInfo);
+			return true;
+		}
+		else
+		{
+			PendingAutoReimportObjects.Remove(InAsset);
 			return true;
 		}
 
@@ -464,7 +514,7 @@ namespace UE::DatasmithImporter
 			return;
 		}
 
-		const bool bHasDirectLinkSourceChanged = RegisteredAutoReimportObjectMap.FindChecked(InAsset)->ExternalSource != UpdatedExternalSource;
+		const bool bHasDirectLinkSourceChanged = AutoReimportObjectsMap.FindChecked(InAsset)->ExternalSource != UpdatedExternalSource;
 		if (bHasDirectLinkSourceChanged)
 		{
 			// The source changed but is still a DirectLink source.
@@ -477,7 +527,7 @@ namespace UE::DatasmithImporter
 	TArray<TSharedRef<FDirectLinkExternalSource>> FDirectLinkManager::GetExternalSourceList() const
 	{
 		TArray<TSharedRef<FDirectLinkExternalSource>> ExternalSources;
-		UriToExternalSourceMap.GenerateValueArray(ExternalSources);
+		DirectLinkSourceToExternalSourceMap.GenerateValueArray(ExternalSources);
 		return ExternalSources;
 	}
 
@@ -527,7 +577,7 @@ namespace UE::DatasmithImporter
 	void FDirectLinkManager::TriggerAutoReimportOnExternalSource(const TSharedRef<FExternalSource>& ExternalSource)
 	{
 		TArray<TSharedRef<FAutoReimportInfo>> AutoReimportInfos;
-		RegisteredAutoReimportExternalSourceMap.MultiFind(ExternalSource, AutoReimportInfos);
+		AutoReimportExternalSourcesMap.MultiFind(ExternalSource, AutoReimportInfos);
 		if (AutoReimportInfos.Num() == 0)
 		{
 			return;
@@ -576,7 +626,7 @@ namespace UE::DatasmithImporter
 		TArray<UObject*> InvalidAssets;
 
 		// We can't call TriggerOnExternalSourceChanged() directly as it may remove items in RegisteredAutoReimportObjectMap while we iterate.
-		for (TPair<UObject*, TSharedRef<FAutoReimportInfo>>& AutoReimportEntry : RegisteredAutoReimportObjectMap)
+		for (TPair<UObject*, TSharedRef<FAutoReimportInfo>>& AutoReimportEntry : AutoReimportObjectsMap)
 		{
 			if (AutoReimportEntry.Value->TargetObject.IsValid())
 			{
diff --git a/Engine/Plugins/Enterprise/DatasmithImporter/Source/DirectLinkExtension/Private/DirectLinkManager.h b/Engine/Plugins/Enterprise/DatasmithImporter/Source/DirectLinkExtension/Private/DirectLinkManager.h
--- a/Engine/Plugins/Enterprise/DatasmithImporter/Source/DirectLinkExtension/Private/DirectLinkManager.h
+++ b/Engine/Plugins/Enterprise/DatasmithImporter/Source/DirectLinkExtension/Private/DirectLinkManager.h
@@ -144,18 +144,20 @@ namespace UE::DatasmithImporter
 
 		TArray<FDirectLinkExternalSourceRegisterInformation> RegisteredExternalSourcesInfo;
 		
-		TMap<FSourceUri, TSharedRef<FDirectLinkExternalSource>> UriToExternalSourceMap;
-		
 		TMap<DirectLink::FSourceHandle, TSharedRef<FDirectLinkExternalSource>> DirectLinkSourceToExternalSourceMap;
 
 		FRWLock ReconnectionListLock;
 		TArray<TSharedRef<FDirectLinkExternalSource>> ExternalSourcesToReconnect;
 
 		TUniquePtr<FDirectLinkAutoReconnectManager> ReconnectionManager;
+		
+		mutable FRWLock AutoReimportLock;
+
+		TSet<TSoftObjectPtr<UObject>> PendingAutoReimportObjects;
 
-		TMap<UObject*, TSharedRef<FAutoReimportInfo>> RegisteredAutoReimportObjectMap;
+		TMap<UObject*, TSharedRef<FAutoReimportInfo>> AutoReimportObjectsMap;
 		
-		TMultiMap<TSharedRef<FExternalSource>, TSharedRef<FAutoReimportInfo>> RegisteredAutoReimportExternalSourceMap;
+		TMultiMap<TSharedRef<FExternalSource>, TSharedRef<FAutoReimportInfo>> AutoReimportExternalSourcesMap;
 
 		TQueue<TSharedPtr<FExternalSource>> PendingReimportQueue;
 
diff --git a/Engine/Plugins/Enterprise/DatasmithImporter/Source/DirectLinkExtension/Private/UI/DirectLinkExtensionStyle.cpp b/Engine/Plugins/Enterprise/DatasmithImporter/Source/DirectLinkExtension/Private/UI/DirectLinkExtensionStyle.cpp
--- a/Engine/Plugins/Enterprise/DatasmithImporter/Source/DirectLinkExtension/Private/UI/DirectLinkExtensionStyle.cpp
+++ b/Engine/Plugins/Enterprise/DatasmithImporter/Source/DirectLinkExtension/Private/UI/DirectLinkExtensionStyle.cpp
@@ -48,6 +48,7 @@ namespace UE::DatasmithImporter
 		StyleSet->Set("DirectLinkExtension.NotAvailable", new IMAGE_PLUGIN_BRUSH(TEXT("Icons/SourceNotAvailable16"), Icon16x16));
 		StyleSet->Set("DirectLinkExtension.OutOfSync", new IMAGE_PLUGIN_BRUSH(TEXT("Icons/SourceOutOfSync16"), Icon16x16));
 		StyleSet->Set("DirectLinkExtension.UpToDate", new IMAGE_PLUGIN_BRUSH(TEXT("Icons/SourceUpToDate16"), Icon16x16));
+		StyleSet->Set("DirectLinkExtension.AutoReimport", new IMAGE_PLUGIN_BRUSH(TEXT("Icons/SourceAutoReimport16"), Icon16x16));
 
 		FSlateStyleRegistry::RegisterSlateStyle(*StyleSet.Get());
 	};
diff --git a/Engine/Plugins/Enterprise/DatasmithImporter/Source/DirectLinkExtension/Private/UI/SDirectLinkAssetThumbnailIndicator.cpp b/Engine/Plugins/Enterprise/DatasmithImporter/Source/DirectLinkExtension/Private/UI/SDirectLinkAssetThumbnailIndicator.cpp
--- a/Engine/Plugins/Enterprise/DatasmithImporter/Source/DirectLinkExtension/Private/UI/SDirectLinkAssetThumbnailIndicator.cpp
+++ b/Engine/Plugins/Enterprise/DatasmithImporter/Source/DirectLinkExtension/Private/UI/SDirectLinkAssetThumbnailIndicator.cpp
@@ -2,6 +2,8 @@
 
 #include "UI/SDirectLinkAssetThumbnailIndicator.h"
 
+#include "DirectLinkExtensionModule.h"
+#include "IDirectLinkManager.h"
 #include "UI/DirectLinkExtensionStyle.h"
 
 #include "ExternalSource.h"
@@ -23,6 +25,7 @@ namespace UE::DatasmithImporter
 			NotAvailable,
 			OutOfSync,
 			UpToDate,
+			AutoReimport,
 		};
 
 		FSourceUri TryGetSourceUri(const FAssetData& AssetData) 
@@ -71,6 +74,16 @@ namespace UE::DatasmithImporter
 				return EDirectLinkSourceStatus::None;
 			}
 
+			// Only loaded asset can possibly be registered for auto-reimport.
+			const bool bLoadAsset = false;
+			if (UObject* Asset = AssetData.FastGetAsset(bLoadAsset))
+			{
+				if (IDirectLinkExtensionModule::Get().GetManager().IsAssetAutoReimportEnabled(Asset))
+				{
+					return EDirectLinkSourceStatus::AutoReimport;
+				}
+			}
+
 			const TSharedPtr<FExternalSource> ExternalSource = IExternalSourceModule::GetOrCreateExternalSource(SourceUri);
 			if (!ExternalSource.IsValid() || !ExternalSource->IsAvailable())
 			{
@@ -97,6 +110,7 @@ namespace UE::DatasmithImporter
 	const FSlateBrush* SDirectLinkAssetThumbnailIndicator::NotAvailableBrush = nullptr;
 	const FSlateBrush* SDirectLinkAssetThumbnailIndicator::OutOfSyncBrush = nullptr;
 	const FSlateBrush* SDirectLinkAssetThumbnailIndicator::UpToDateBrush = nullptr;
+	const FSlateBrush* SDirectLinkAssetThumbnailIndicator::AutoReimportBrush = nullptr;
 
 	void SDirectLinkAssetThumbnailIndicator::Construct(const FArguments& InArgs)
 	{
@@ -129,6 +143,7 @@ namespace UE::DatasmithImporter
 			NotAvailableBrush = FDirectLinkExtensionStyle::Get().GetBrush(TEXT("DirectLinkExtension.NotAvailable"));
 			OutOfSyncBrush = FDirectLinkExtensionStyle::Get().GetBrush(TEXT("DirectLinkExtension.OutOfSync"));
 			UpToDateBrush = FDirectLinkExtensionStyle::Get().GetBrush(TEXT("DirectLinkExtension.UpToDate"));
+			AutoReimportBrush = FDirectLinkExtensionStyle::Get().GetBrush(TEXT("DirectLinkExtension.AutoReimport"));
 		}
 	}
 
@@ -147,6 +162,8 @@ namespace UE::DatasmithImporter
 			return OutOfSyncBrush;
 		case ThumbnailIndicatorImpl::EDirectLinkSourceStatus::UpToDate:
 			return UpToDateBrush;
+		case ThumbnailIndicatorImpl::EDirectLinkSourceStatus::AutoReimport:
+			return AutoReimportBrush;
 		case ThumbnailIndicatorImpl::EDirectLinkSourceStatus::None:
 		default:
 			return nullptr;
@@ -195,6 +212,8 @@ namespace UE::DatasmithImporter
 			return LOCTEXT("OutOfSyncTooltip", "Asset is not synced to Direct Link source.");
 		case ThumbnailIndicatorImpl::EDirectLinkSourceStatus::UpToDate:
 			return LOCTEXT("NotAvailableTooltip", "Asset is up-to-date with Direct Link source.");
+		case ThumbnailIndicatorImpl::EDirectLinkSourceStatus::AutoReimport:
+			return LOCTEXT("AutoReimportTooltip", "Asset will auto-reimport on Direct Link source change.");
 		case ThumbnailIndicatorImpl::EDirectLinkSourceStatus::None:
 		default:
 			return FText();
diff --git a/Engine/Plugins/Enterprise/DatasmithImporter/Source/DirectLinkExtension/Private/UI/SDirectLinkAssetThumbnailIndicator.h b/Engine/Plugins/Enterprise/DatasmithImporter/Source/DirectLinkExtension/Private/UI/SDirectLinkAssetThumbnailIndicator.h
--- a/Engine/Plugins/Enterprise/DatasmithImporter/Source/DirectLinkExtension/Private/UI/SDirectLinkAssetThumbnailIndicator.h
+++ b/Engine/Plugins/Enterprise/DatasmithImporter/Source/DirectLinkExtension/Private/UI/SDirectLinkAssetThumbnailIndicator.h
@@ -38,6 +38,7 @@ namespace UE::DatasmithImporter
 		static const FSlateBrush* NotAvailableBrush;
 		static const FSlateBrush* OutOfSyncBrush;
 		static const FSlateBrush* UpToDateBrush;
+		static const FSlateBrush* AutoReimportBrush;
 
 		/** Asset path for this indicator widget.*/
 		FAssetData AssetData;
