Fix Non-Unity build - missing includes

diff --git a/Engine/Plugins/Enterprise/DatasmithImporter/Source/DirectLinkExtension/Private/DirectLinkManager.h b/Engine/Plugins/Enterprise/DatasmithImporter/Source/DirectLinkExtension/Private/DirectLinkManager.h
--- a/Engine/Plugins/Enterprise/DatasmithImporter/Source/DirectLinkExtension/Private/DirectLinkManager.h
+++ b/Engine/Plugins/Enterprise/DatasmithImporter/Source/DirectLinkExtension/Private/DirectLinkManager.h
@@ -12,6 +12,7 @@
 #include "DirectLinkEndpoint.h"
 #include "HAL/CriticalSection.h"
 #include "IUriResolver.h"
+#include "UObject/SoftObjectPtr.h"
 
 struct FAssetData;
 
