Datasmith Landscape Import Issue reported by [REDACTED_PERSON] from CityEngine

diff --git a/Engine/Plugins/Enterprise/DatasmithImporter/Source/DatasmithImporter/Private/DatasmithLandscapeImporter.cpp b/Engine/Plugins/Enterprise/DatasmithImporter/Source/DatasmithImporter/Private/DatasmithLandscapeImporter.cpp
--- a/Engine/Plugins/Enterprise/DatasmithImporter/Source/DatasmithImporter/Private/DatasmithLandscapeImporter.cpp
+++ b/Engine/Plugins/Enterprise/DatasmithImporter/Source/DatasmithImporter/Private/DatasmithLandscapeImporter.cpp
@@ -58,7 +58,7 @@ AActor* FDatasmithLandscapeImporter::ImportLandscapeActor( const TSharedRef< IDa
 	FLandscapeImportHelper::TransformHeightmapImportData(ImportData, FinalHeightData, OutImportDescriptor.ImportResolutions[DescriptorIndex], FLandscapeImportResolution(SizeX, SizeY), ELandscapeImportTransformType::ExpandCentered);
 		
 	const FVector Offset = FTransform( LandscapeActorElement->GetRotation(), FVector::ZeroVector,
-		LandscapeActorElement->GetScale() ).TransformVector( FVector( -OutComponentCount.X * QuadsPerComponent / 2, -OutComponentCount.Y * QuadsPerComponent / 2, 0 ) );
+		LandscapeActorElement->GetScale() ).TransformVector( FVector( -OutComponentCount.X * QuadsPerComponent / 2., -OutComponentCount.Y * QuadsPerComponent / 2., 0. ) );
 
 	FVector OriginalTranslation = LandscapeActorElement->GetTranslation();
 	FVector OriginalScale = LandscapeActorElement->GetScale();
