DirectLink Editor - Adding refresh button to DirectLink source selection dialog and fixing button right-allignment.

diff --git a/Engine/Plugins/Enterprise/DatasmithImporter/Source/DirectLinkExtension/DirectLinkExtension.Build.cs b/Engine/Plugins/Enterprise/DatasmithImporter/Source/DirectLinkExtension/DirectLinkExtension.Build.cs
--- a/Engine/Plugins/Enterprise/DatasmithImporter/Source/DirectLinkExtension/DirectLinkExtension.Build.cs
+++ b/Engine/Plugins/Enterprise/DatasmithImporter/Source/DirectLinkExtension/DirectLinkExtension.Build.cs
@@ -12,6 +12,7 @@ namespace UnrealBuildTool.Rules
 				new string[]
 				{
 					"ContentBrowserData",
+					"EditorStyle",
 					"EditorWidgets",
 					"Engine",
 					"InputCore",
diff --git a/Engine/Plugins/Enterprise/DatasmithImporter/Source/DirectLinkExtension/Private/UI/SDirectLinkAvailableSource.cpp b/Engine/Plugins/Enterprise/DatasmithImporter/Source/DirectLinkExtension/Private/UI/SDirectLinkAvailableSource.cpp
--- a/Engine/Plugins/Enterprise/DatasmithImporter/Source/DirectLinkExtension/Private/UI/SDirectLinkAvailableSource.cpp
+++ b/Engine/Plugins/Enterprise/DatasmithImporter/Source/DirectLinkExtension/Private/UI/SDirectLinkAvailableSource.cpp
@@ -6,6 +6,7 @@
 #include "DirectLinkExtensionModule.h"
 #include "IDirectLinkManager.h"
 
+#include "Styling/SlateBrush.h"
 #include "Widgets/Input/SButton.h"
 #include "Widgets/Text/STextBlock.h"
 #include "Widgets/Views/SListView.h"
@@ -144,6 +145,27 @@ namespace UE::DatasmithImporter
 			.VAlign(VAlign_Bottom)
 			[
 				SNew(SHorizontalBox)
+				+SHorizontalBox::Slot()
+				.AutoWidth()
+				.HAlign(HAlign_Left)
+				[
+					SNew(SButton)
+					.ButtonStyle(FAppStyle::Get(), "SimpleButton")
+					.ContentPadding(1)
+					.ToolTipText(LOCTEXT("RefreshSourceListButtonTooltip", "Refresh the list of Direct Link sources."))
+					.OnClicked_Lambda([this]() { GenerateDirectLinkExternalSourceInfos(); return FReply::Handled(); })
+					[
+						SNew(SImage)
+						.Image(FEditorStyle::GetBrush("Icons.Refresh"))
+					]
+				]
+
+				+SHorizontalBox::Slot()
+				.FillWidth(1.f)
+				[
+					SNew(SSpacer)
+				]
+
 				+ SHorizontalBox::Slot()
 				.Padding(5, 0)
 				.AutoWidth()
@@ -190,6 +212,7 @@ namespace UE::DatasmithImporter
 
 	void SDirectLinkAvailableSource::GenerateDirectLinkExternalSourceInfos()
 	{
+		DirectLinkExternalSourceInfos.Empty(DirectLinkExternalSourceInfos.Num());
 		for (const TSharedRef<FDirectLinkExternalSource>& ExternalSource : IDirectLinkExtensionModule::Get().GetManager().GetExternalSourceList())
 		{
 			DirectLinkExternalSourceInfos.Add(MakeShared<FDirectLinkExternalSourceInfo>(ExternalSource));
