DirectLink in Editor - Fix autoreimport enabled after a reimport.

diff --git a/Engine/Plugins/Enterprise/DatasmithImporter/Source/DatasmithImporter/Private/DatasmithImportFactory.cpp b/Engine/Plugins/Enterprise/DatasmithImporter/Source/DatasmithImporter/Private/DatasmithImportFactory.cpp
--- a/Engine/Plugins/Enterprise/DatasmithImporter/Source/DatasmithImporter/Private/DatasmithImportFactory.cpp
+++ b/Engine/Plugins/Enterprise/DatasmithImporter/Source/DatasmithImporter/Private/DatasmithImportFactory.cpp
@@ -561,7 +561,7 @@ bool UDatasmithImportFactory::Import( FDatasmithImportContext& ImportContext )
 	{
 		ImportContext.DisplayMessages();
 
-		if ( bImportSuccess && !bOutOperationCanceled )
+		if ( bImportSuccess && !bOutOperationCanceled && !ImportContext.bIsAReimport )
 		{
 			IDatasmithContentEditorModule& ContentEditorModule = IDatasmithContentEditorModule::Get();
 			const bool bIsAutoReimportAvailable = ContentEditorModule.IsAssetAutoReimportAvailable( ImportContext.SceneAsset ).Get( false );
