Fix for static analysis warning
https://horde.devtools.epicgames.com/log/61e94c0f276892ce108f9207?lineindex=6744

diff --git a/Engine/Plugins/Enterprise/DatasmithImporter/Source/DatasmithImporter/Private/DatasmithConsumer.cpp b/Engine/Plugins/Enterprise/DatasmithImporter/Source/DatasmithImporter/Private/DatasmithConsumer.cpp
--- a/Engine/Plugins/Enterprise/DatasmithImporter/Source/DatasmithImporter/Private/DatasmithConsumer.cpp
+++ b/Engine/Plugins/Enterprise/DatasmithImporter/Source/DatasmithImporter/Private/DatasmithConsumer.cpp
@@ -1069,7 +1069,7 @@ ULevel* UDatasmithConsumer::FindOrAddLevel(const FString& InLevelName)
 	// If it does not exist, create a new one and add it to the working world
 	else
 	{
-		StreamingLevel = EditorLevelUtils::CreateNewStreamingLevelForWorld( *WorkingWorld, ULevelStreamingAlwaysLoaded::StaticClass(), *PackageFilename, false, nullptr, false );
+		StreamingLevel = EditorLevelUtils::CreateNewStreamingLevelForWorld( *WorkingWorld, ULevelStreamingAlwaysLoaded::StaticClass(), PackageFilename, false, nullptr, false );
 
 		if( StreamingLevel )
 		{
