Datasmith: Fix arealight mobility not correct when an animation is bound on it

diff --git a/Engine/Plugins/Enterprise/DatasmithImporter/Source/DatasmithImporter/Private/DatasmithImporter.cpp b/Engine/Plugins/Enterprise/DatasmithImporter/Source/DatasmithImporter/Private/DatasmithImporter.cpp
--- a/Engine/Plugins/Enterprise/DatasmithImporter/Source/DatasmithImporter/Private/DatasmithImporter.cpp
+++ b/Engine/Plugins/Enterprise/DatasmithImporter/Source/DatasmithImporter/Private/DatasmithImporter.cpp
@@ -1433,7 +1433,11 @@ void FDatasmithImporter::ImportLevelSequences( FDatasmithImportContext& ImportCo
 			if (ImportContext.SceneTranslator)
 			{
 				FDatasmithLevelSequencePayload LevelSequencePayload;
-				ImportContext.SceneTranslator->LoadLevelSequence(SequenceElement.ToSharedRef(), LevelSequencePayload);
+				if (!ImportContext.SceneTranslator->LoadLevelSequence(SequenceElement.ToSharedRef(), LevelSequencePayload))
+				{
+					FString ErrorMessage = LOCTEXT("FailedToLoad", "Failed to load some animation sequences: ").ToString() + SequenceElement->GetName() + TEXT("\n");
+					ImportContext.LogError(FText::FromString(ErrorMessage));
+				}
 			}
 
 			ULevelSequence*& ImportedLevelSequence = ImportContext.ImportedLevelSequences.FindOrAdd( SequenceElement.ToSharedRef() );
diff --git a/Engine/Plugins/Enterprise/DatasmithImporter/Source/DatasmithImporter/Private/DatasmithLevelSequenceImporter.cpp b/Engine/Plugins/Enterprise/DatasmithImporter/Source/DatasmithImporter/Private/DatasmithLevelSequenceImporter.cpp
--- a/Engine/Plugins/Enterprise/DatasmithImporter/Source/DatasmithImporter/Private/DatasmithLevelSequenceImporter.cpp
+++ b/Engine/Plugins/Enterprise/DatasmithImporter/Source/DatasmithImporter/Private/DatasmithLevelSequenceImporter.cpp
@@ -3,6 +3,7 @@
 #include "DatasmithLevelSequenceImporter.h"
 
 #include "DatasmithAnimationElements.h"
+#include "DatasmithAreaLightActor.h"
 #include "DatasmithDefinitions.h"
 #include "DatasmithImportContext.h"
 #include "DatasmithSceneActor.h"
@@ -377,6 +378,11 @@ namespace DatasmithLevelSequenceImporterImpl
 		{
 			SceneComp->SetMobility(EComponentMobility::Movable);
 		}
+
+		if (ADatasmithAreaLightActor* DatasmithAreaLightActor = Cast<ADatasmithAreaLightActor>(Actor))
+		{
+			DatasmithAreaLightActor->Mobility = EComponentMobility::Movable;
+		}
 	}
 
 	void PopulateVisibilityTrack(const TSharedRef<IDatasmithVisibilityAnimationElement>& VisibilityAnimation, int32 AnimIndex, FDatasmithImportContext& ImportContext, ULevelSequence* LevelSequence, FFrameNumber& MinFrameNumber, FFrameNumber& MaxFrameNumber)
diff --git a/Engine/Plugins/Enterprise/DatasmithImporter/Source/DatasmithNativeTranslator/Private/DatasmithNativeTranslator.cpp b/Engine/Plugins/Enterprise/DatasmithImporter/Source/DatasmithNativeTranslator/Private/DatasmithNativeTranslator.cpp
--- a/Engine/Plugins/Enterprise/DatasmithImporter/Source/DatasmithNativeTranslator/Private/DatasmithNativeTranslator.cpp
+++ b/Engine/Plugins/Enterprise/DatasmithImporter/Source/DatasmithNativeTranslator/Private/DatasmithNativeTranslator.cpp
@@ -57,7 +57,7 @@ namespace DatasmithNativeTranslatorImpl
 		{
 			// Make sure the new UDatasmithMesh object is not created while a garbage collection is performed
 			FGCScopeGuard GCGuard;
-			// Setting the RF_Standalone bitmask on the new UDatasmithMesh object, to make sure it is not garbage collected 
+			// Setting the RF_Standalone bitmask on the new UDatasmithMesh object, to make sure it is not garbage collected
 			// while loading and processing the udsmesh file. This can happen on very big meshes (5M+ triangles)
 			DatasmithMesh = NewObject< UDatasmithMesh >(GetTransientPackage(), NAME_None, RF_Standalone);
 		}
@@ -126,7 +126,7 @@ namespace DatasmithNativeTranslatorImpl
 		DatasmithMeshHelper::PrepareAttributeForStaticMesh(MeshDescription);
 
 		// Do not compute normals and tangents during conversion since we have other operations to apply
-		// on the mesh that might invalidate them anyway and we must also validate the mesh to detect 
+		// on the mesh that might invalidate them anyway and we must also validate the mesh to detect
 		// vertex positions containing NaN before doing computation as MikkTSpace crashes on NaN values.
 		const bool bSkipNormalsAndTangents = true;
 		FStaticMeshOperations::ConvertFromRawMesh(RawMesh, MeshDescription, GroupNamePerGroupIndex, bSkipNormalsAndTangents);
@@ -190,7 +190,7 @@ bool FDatasmithNativeTranslator::LoadLevelSequence(const TSharedRef<IDatasmithLe
 	FDatasmithAnimationSerializer AnimSerializer;
 	if (LevelSequenceElement->GetFile() && IFileManager::Get().FileExists(LevelSequenceElement->GetFile()))
 	{
-		AnimSerializer.Deserialize(LevelSequenceElement, LevelSequenceElement->GetFile());
+		return AnimSerializer.Deserialize(LevelSequenceElement, LevelSequenceElement->GetFile());
 	}
-	return true;
+	return false;
 }
diff --git a/Engine/Plugins/Enterprise/DatasmithImporter/Source/DatasmithTranslator/Public/DatasmithTranslator.h b/Engine/Plugins/Enterprise/DatasmithImporter/Source/DatasmithTranslator/Public/DatasmithTranslator.h
--- a/Engine/Plugins/Enterprise/DatasmithImporter/Source/DatasmithTranslator/Public/DatasmithTranslator.h
+++ b/Engine/Plugins/Enterprise/DatasmithImporter/Source/DatasmithTranslator/Public/DatasmithTranslator.h
@@ -48,13 +48,13 @@ public:
 	virtual FName GetFName() const = 0;
 
 	/**
-	 * Called on the CDO to initialize this Translator implementation
+	 * Called to initialize this Translator instance
 	 * @param OutCapabilities	Declared capabilities for this Translator
 	 */
 	virtual void Initialize(FDatasmithTranslatorCapabilities& OutCapabilities) {}
 
 	/**
-	 * Called on the CDO, additional validation step to ensure the given file is supported (Optional)
+	 * Additional validation step to ensure the given file is supported (Optional)
 	 *
 	 * @param Source	Candidate Source
 	 * @returns			True if the given source is supported by this Translator
