Fix missing include.
#rnx

diff --git a/Engine/Plugins/Enterprise/DatasmithImporter/Source/DirectLinkExtension/Private/UI/SDirectLinkAvailableSource.cpp b/Engine/Plugins/Enterprise/DatasmithImporter/Source/DirectLinkExtension/Private/UI/SDirectLinkAvailableSource.cpp
--- a/Engine/Plugins/Enterprise/DatasmithImporter/Source/DirectLinkExtension/Private/UI/SDirectLinkAvailableSource.cpp
+++ b/Engine/Plugins/Enterprise/DatasmithImporter/Source/DirectLinkExtension/Private/UI/SDirectLinkAvailableSource.cpp
@@ -8,6 +8,7 @@
 
 #include "Styling/SlateBrush.h"
 #include "Widgets/Input/SButton.h"
+#include "Widgets/Input/SScrollBar.h"
 #include "Widgets/Text/STextBlock.h"
 #include "Widgets/Views/SListView.h"
 #include "Widgets/Views/STableRow.h"
