Datasmith Importer - Fix crash on reimport when reparenting component with nameclash.

diff --git a/Engine/Plugins/Enterprise/DatasmithImporter/Source/DatasmithImporter/Private/Utility/DatasmithImporterImpl.cpp b/Engine/Plugins/Enterprise/DatasmithImporter/Source/DatasmithImporter/Private/Utility/DatasmithImporterImpl.cpp
--- a/Engine/Plugins/Enterprise/DatasmithImporter/Source/DatasmithImporter/Private/Utility/DatasmithImporterImpl.cpp
+++ b/Engine/Plugins/Enterprise/DatasmithImporter/Source/DatasmithImporter/Private/Utility/DatasmithImporterImpl.cpp
@@ -620,6 +620,17 @@ USceneComponent* FDatasmithImporterImpl::FinalizeSceneComponent(FDatasmithImport
 					}
 				}
 			}
+
+			// If a DestinationComponent was found using DatasmithElementId (without using the name), make sure that there is no other component with the SourceComponent Name.
+			if (DestinationComponent)
+			{
+				if (UActorComponent* OldComponent = static_cast<UActorComponent*>(FindObjectWithOuter(&DestinationActor, UActorComponent::StaticClass(), SourceComponent.GetFName())))
+				{
+					// It's okay to rename to null at this point as either the OldComponent was flagged for deletion by DeleteNonImportedDatasmithElementFromSceneActor()
+					// at this point OR this component is in the reimported scene and we'll find it back using the the DatasmithElementId.
+					OldComponent->Rename(nullptr, OldComponent->GetOuter(), REN_DontCreateRedirectors | REN_NonTransactional);
+				}
+			}
 		}
 
 		TArray< FMigratedTemplatePairType > MigratedTemplates = MigrateTemplates(
