Adding report on tessellator in Analytics

diff --git a/Engine/Plugins/Enterprise/DatasmithCADImporter/Source/ParametricSurface/Public/ParametricSurfaceTranslator.h b/Engine/Plugins/Enterprise/DatasmithCADImporter/Source/ParametricSurface/Public/ParametricSurfaceTranslator.h
--- a/Engine/Plugins/Enterprise/DatasmithCADImporter/Source/ParametricSurface/Public/ParametricSurfaceTranslator.h
+++ b/Engine/Plugins/Enterprise/DatasmithCADImporter/Source/ParametricSurface/Public/ParametricSurfaceTranslator.h
@@ -4,6 +4,8 @@
 
 #include "CoreMinimal.h"
 
+#include "CADOptions.h"
+
 #include "DatasmithImportOptions.h"
 #include "DatasmithTranslator.h"
 #include "UObject/ObjectMacros.h"
@@ -20,6 +22,12 @@ class PARAMETRICSURFACE_API FParametricSurfaceTranslator : public IDatasmithTran
 {
 public:
 
+	FParametricSurfaceTranslator()
+	{
+		// Initialize bUseCADKernel with current value of CVar ds.CADTranslator.DisableCADKernelTessellation
+		CommonTessellationOptions.bUseCADKernel = !CADLibrary::FImportParameters::bGDisableCADKernelTessellation;
+	}
+
 	// Begin IDatasmithTranslator overrides
 	virtual void GetSceneImportOptions(TArray<TStrongObjectPtr<UDatasmithOptionsBase>>& Options) override;
 	virtual void SetSceneImportOptions(TArray<TStrongObjectPtr<UDatasmithOptionsBase>>& Options) override;
diff --git a/Engine/Plugins/Enterprise/DatasmithContent/Source/DatasmithContent/Public/DatasmithImportOptions.h b/Engine/Plugins/Enterprise/DatasmithContent/Source/DatasmithContent/Public/DatasmithImportOptions.h
--- a/Engine/Plugins/Enterprise/DatasmithContent/Source/DatasmithContent/Public/DatasmithImportOptions.h
+++ b/Engine/Plugins/Enterprise/DatasmithContent/Source/DatasmithContent/Public/DatasmithImportOptions.h
@@ -261,6 +261,8 @@ struct DATASMITHCONTENT_API FDatasmithTessellationOptions
 	UPROPERTY(config, EditAnywhere, BlueprintReadWrite, AdvancedDisplay, Category = "Geometry & Tessellation Options", meta = (ToolTip = "Stitching technique applied on model before tessellation. Sewing could impact number of objects."))
 	EDatasmithCADStitchingTechnique StitchingTechnique;
 
+	bool bUseCADKernel = false;
+
 public:
 	bool operator == (const FDatasmithTessellationOptions& Other) const
 	{
diff --git a/Engine/Plugins/Enterprise/DatasmithImporter/Source/DatasmithImporter/Private/DatasmithImportFactory.cpp b/Engine/Plugins/Enterprise/DatasmithImporter/Source/DatasmithImporter/Private/DatasmithImportFactory.cpp
--- a/Engine/Plugins/Enterprise/DatasmithImporter/Source/DatasmithImporter/Private/DatasmithImportFactory.cpp
+++ b/Engine/Plugins/Enterprise/DatasmithImporter/Source/DatasmithImporter/Private/DatasmithImportFactory.cpp
@@ -333,8 +333,32 @@ namespace DatasmithImportFactoryImpl
 			EventAttributes.Emplace( TEXT("ExporterOS"), ImportContext.Scene->GetUserOS() );
 			EventAttributes.Emplace( TEXT("ExportDuration"), ImportContext.Scene->GetExportDuration() );
 
-			FString SourceFileExtension = ImportContext.SceneTranslator ? ImportContext.SceneTranslator->GetSource().GetSourceFileExtension() : TEXT("Unknown");
-			EventAttributes.Emplace( TEXT("SourceFileExtension"), SourceFileExtension );
+			if (ImportContext.SceneTranslator)
+			{
+				EventAttributes.Emplace(TEXT("SourceFileExtension"), ImportContext.SceneTranslator->GetSource().GetSourceFileExtension());
+
+				// Log tessellator if CADKernel has been used
+				bool bUseCADKernel = false;
+				TArray<TStrongObjectPtr<UDatasmithOptionsBase>> Options;
+				ImportContext.SceneTranslator->GetSceneImportOptions(Options);
+
+				for (const TStrongObjectPtr<UDatasmithOptionsBase>& Option : Options)
+				{
+					if (UDatasmithCommonTessellationOptions* TessellationOptionsObject = Cast<UDatasmithCommonTessellationOptions>(Option.Get()))
+					{
+						bUseCADKernel = TessellationOptionsObject->Options.bUseCADKernel;
+					}
+				}
+
+				if (bUseCADKernel)
+				{
+					EventAttributes.Emplace(TEXT("Tessellator"), TEXT("CADKernel"));
+				}
+			}
+			else
+			{
+				EventAttributes.Emplace(TEXT("SourceFileExtension"), TEXT("Unknown"));
+			}
 
 			FString EventText = TEXT("Datasmith.");
 			EventText += ImportContext.bIsAReimport ? TEXT("Reimport") : TEXT("Import");
