Fixing duplicated loctext entry.

diff --git a/Engine/Plugins/Enterprise/DatasmithImporter/Source/DirectLinkExtension/Private/UI/SDirectLinkAssetThumbnailIndicator.cpp b/Engine/Plugins/Enterprise/DatasmithImporter/Source/DirectLinkExtension/Private/UI/SDirectLinkAssetThumbnailIndicator.cpp
--- a/Engine/Plugins/Enterprise/DatasmithImporter/Source/DirectLinkExtension/Private/UI/SDirectLinkAssetThumbnailIndicator.cpp
+++ b/Engine/Plugins/Enterprise/DatasmithImporter/Source/DirectLinkExtension/Private/UI/SDirectLinkAssetThumbnailIndicator.cpp
@@ -211,7 +211,7 @@ namespace UE::DatasmithImporter
 		case ThumbnailIndicatorImpl::EDirectLinkSourceStatus::OutOfSync:
 			return LOCTEXT("OutOfSyncTooltip", "Asset is not synced to Direct Link source.");
 		case ThumbnailIndicatorImpl::EDirectLinkSourceStatus::UpToDate:
-			return LOCTEXT("NotAvailableTooltip", "Asset is up-to-date with Direct Link source.");
+			return LOCTEXT("UpToDateTooltip", "Asset is up-to-date with Direct Link source.");
 		case ThumbnailIndicatorImpl::EDirectLinkSourceStatus::AutoReimport:
 			return LOCTEXT("AutoReimportTooltip", "Asset will auto-reimport on Direct Link source change.");
 		case ThumbnailIndicatorImpl::EDirectLinkSourceStatus::None:
