DirectLink in Editor - Fixing async task never stopped and preventing UE from closing.

Auto-reconnection task will re-enqueue another reconnection task after running if there are still DirectLink sources to reconnect to. Normally this task would be stopped on module shutdown, but it turns out that the presence of running background tasks prevents Unreal from shutting down, and thus the module was never shutdown. To fix this, we now listen for the on FCoreDelegate::OnPreExit event, and stop any running reconnection task.

diff --git a/Engine/Plugins/Enterprise/DatasmithImporter/Source/DirectLinkExtension/Private/DirectLinkManager.cpp b/Engine/Plugins/Enterprise/DatasmithImporter/Source/DirectLinkExtension/Private/DirectLinkManager.cpp
--- a/Engine/Plugins/Enterprise/DatasmithImporter/Source/DirectLinkExtension/Private/DirectLinkManager.cpp
+++ b/Engine/Plugins/Enterprise/DatasmithImporter/Source/DirectLinkExtension/Private/DirectLinkManager.cpp
@@ -11,6 +11,7 @@
 #include "AssetRegistry/IAssetRegistry.h"
 #include "Async/Async.h"
 #include "EditorReimportHandler.h"
+#include "Misc/CoreDelegates.h"
 #include "UObject/UObjectGlobals.h"
 
 #if WITH_EDITOR
@@ -33,6 +34,18 @@ namespace UE::DatasmithImporter
 			bAutoReconnectEnabled = DefaultSettings->bAutoReconnect;
 			ReconnectionDelayInSeconds = DefaultSettings->ReconnectionDelayInSeconds;
 		}
+
+		FCoreDelegates::OnPreExit.AddRaw(this, &FDirectLinkAutoReconnectManager::Stop);
+	}
+
+	FDirectLinkAutoReconnectManager::~FDirectLinkAutoReconnectManager()
+	{
+		FCoreDelegates::OnPreExit.RemoveAll(this);
+		if (CompletedFuture.IsValid())
+		{
+			Stop();
+			CompletedFuture.Wait();
+		}
 	}
 
 	bool FDirectLinkAutoReconnectManager::Start()
diff --git a/Engine/Plugins/Enterprise/DatasmithImporter/Source/DirectLinkExtension/Private/DirectLinkManager.h b/Engine/Plugins/Enterprise/DatasmithImporter/Source/DirectLinkExtension/Private/DirectLinkManager.h
--- a/Engine/Plugins/Enterprise/DatasmithImporter/Source/DirectLinkExtension/Private/DirectLinkManager.h
+++ b/Engine/Plugins/Enterprise/DatasmithImporter/Source/DirectLinkExtension/Private/DirectLinkManager.h
@@ -31,14 +31,7 @@ namespace UE::DatasmithImporter
 	public:
 		FDirectLinkAutoReconnectManager(FDirectLinkManager& InManager);
 
-		~FDirectLinkAutoReconnectManager()
-		{
-			if (CompletedFuture.IsValid())
-			{
-				Stop();
-				CompletedFuture.Wait();
-			}
-		}
+		~FDirectLinkAutoReconnectManager();
 
 		bool Start();
 
