Direct Link in Editor - Fixing crash happening when auto-reimport is triggered while we are auto-saving.

The auto-save's "FSlowTask update" triggers a Slate update which can in turn execute pending task on the MainThread. This could cause a DirectLink auto-reimport task to execute while saving, causing a crash. To fix this we now make sure that we are in a valid context before triggering a reimport.

diff --git a/Engine/Plugins/Enterprise/DatasmithImporter/Source/DirectLinkExtension/Private/DirectLinkManager.cpp b/Engine/Plugins/Enterprise/DatasmithImporter/Source/DirectLinkExtension/Private/DirectLinkManager.cpp
--- a/Engine/Plugins/Enterprise/DatasmithImporter/Source/DirectLinkExtension/Private/DirectLinkManager.cpp
+++ b/Engine/Plugins/Enterprise/DatasmithImporter/Source/DirectLinkExtension/Private/DirectLinkManager.cpp
@@ -573,6 +573,12 @@ namespace UE::DatasmithImporter
 
 		Async(EAsyncExecution::TaskGraphMainThread, [this]() {
 
+			if (UE::IsSavingPackage(nullptr) || IsGarbageCollecting())
+			{
+				// Can't trigger a reimport while we are saving or garbage collecting.
+				return;
+			}
+
 			TSet<TSharedRef<FExternalSource>> PendingReimportSet;
 			TSharedPtr<FExternalSource> EnqueuedSourceToReimport;
 			while (PendingReimportQueue.Dequeue(EnqueuedSourceToReimport))
