remove hard-coded TMGS_Sharpen5 mip filter from DataSmithTextureImporter

diff --git a/Engine/Plugins/Enterprise/DatasmithImporter/Source/DatasmithImporter/Private/DatasmithTextureImporter.cpp b/Engine/Plugins/Enterprise/DatasmithImporter/Source/DatasmithImporter/Private/DatasmithTextureImporter.cpp
--- a/Engine/Plugins/Enterprise/DatasmithImporter/Source/DatasmithImporter/Private/DatasmithTextureImporter.cpp
+++ b/Engine/Plugins/Enterprise/DatasmithImporter/Source/DatasmithImporter/Private/DatasmithTextureImporter.cpp
@@ -218,7 +218,6 @@ UTexture* FDatasmithTextureImporter::CreateTexture(const TSharedPtr<IDatasmithTe
 	switch (TextureElement->GetTextureMode())
 	{
 	case EDatasmithTextureMode::Diffuse:
-		TextureFact->MipGenSettings = TMGS_Sharpen5;
 		TextureFact->LODGroup = TEXTUREGROUP_World;
 		break;
 	case EDatasmithTextureMode::Specular:
@@ -451,7 +450,6 @@ void UDatasmithTexturePipeline::ExecutePreImportPipeline(UInterchangeBaseNodeCon
 	switch ( TextureElement->GetTextureMode() )
 	{
 	case EDatasmithTextureMode::Diffuse:
-		MipGenSettings = TMGS_Sharpen5;
 		LODGroup = TEXTUREGROUP_World;
 		break;
 	case EDatasmithTextureMode::Specular:
