Disable interchange texture import in Datasmith, when enabled it creates a deadlock between the Interchange import task and the mainthread.
This shall be replaced by a proper implementation of the Interchange framework for Datasmith.

diff --git a/Engine/Plugins/Enterprise/DatasmithImporter/Source/DatasmithImporter/Private/DatasmithImporter.cpp b/Engine/Plugins/Enterprise/DatasmithImporter/Source/DatasmithImporter/Private/DatasmithImporter.cpp
--- a/Engine/Plugins/Enterprise/DatasmithImporter/Source/DatasmithImporter/Private/DatasmithImporter.cpp
+++ b/Engine/Plugins/Enterprise/DatasmithImporter/Source/DatasmithImporter/Private/DatasmithImporter.cpp
@@ -505,8 +505,9 @@ void FDatasmithImporter::ImportTextures( FDatasmithImportContext& ImportContext
 
 		FDatasmithTextureResize::Initialize();
 
-		// Try to import textures async first, if possible
-		if ( GetDefault<UEditorExperimentalSettings>()->bEnableInterchangeFramework )
+		// Try to import textures async first, this does not work with DatasmithImportFactory as it creates a deadlock on the main thread.
+		const bool bInterchangeImport = false; /*GetDefault<UEditorExperimentalSettings>()->bEnableInterchangeFramework*/
+		if ( bInterchangeImport )
 		{
 			for ( int32 TextureIndex = 0; TextureIndex < FilteredTextureElements.Num(); TextureIndex++ )
 			{
