Datasmith Interchange Integration Prototype

Adding crude support for importing a DatasmithScene asset and Textures via the Interchange framework.
More work is needed to have proper node dependencies, and reimport data for those assets.
The rest of the assets (meshes, materials, etc) are imported sequencially during the PostImportPipeline step.

diff --git a/Engine/Plugins/Enterprise/DatasmithImporter/Source/DatasmithImporter/Private/DatasmithImporterModule.cpp b/Engine/Plugins/Enterprise/DatasmithImporter/Source/DatasmithImporter/Private/DatasmithImporterModule.cpp
--- a/Engine/Plugins/Enterprise/DatasmithImporter/Source/DatasmithImporter/Private/DatasmithImporterModule.cpp
+++ b/Engine/Plugins/Enterprise/DatasmithImporter/Source/DatasmithImporter/Private/DatasmithImporterModule.cpp
@@ -116,6 +116,11 @@ public:
 		ResetFromTemplates( Object );
 	}
 
+	virtual FOnGenerateDatasmithImportMenu& OnGenerateDatasmithImportMenu() override
+	{
+		return GenerateDatasmithImportMenuEvent;
+	}
+
 	bool IsInOfflineMode() const
 	{
 		return GetDefault< UDatasmithImporterEditorSettings >() && GetDefault< UDatasmithImporterEditorSettings >()->bOfflineImporter;
@@ -169,6 +174,9 @@ private:
 	FDelegateHandle IsAssetAutoReimportEnabledHandle;
 	FDelegateHandle BrowseExternalSourceUriHandle;
 	FDelegateHandle GetSupporedUriSchemeHandle;
+
+	FOnGenerateDatasmithImportMenu GenerateDatasmithImportMenuEvent;
+
 };
 
 void FDatasmithImporterModule::SetupMenuEntry()
@@ -205,6 +213,8 @@ void FDatasmithImporterModule::SetupMenuEntry()
 					FSlateIcon( FDatasmithStyle::GetStyleSetName(), TEXT( "Datasmith.Import" ) ),
 					FUIAction( FExecuteAction::CreateRaw( this, &FDatasmithImporterModule::OnClickedImportDirectLinkMenuEntry ) )
 				);
+
+				GenerateDatasmithImportMenuEvent.Broadcast(SubSection);
 			}),
 			bOpenMenuOnClick,
 			FSlateIcon( FDatasmithStyle::GetStyleSetName(), TEXT( "Datasmith.Import" ) )
diff --git a/Engine/Plugins/Enterprise/DatasmithImporter/Source/DatasmithImporter/Public/DatasmithImporterModule.h b/Engine/Plugins/Enterprise/DatasmithImporter/Source/DatasmithImporter/Public/DatasmithImporterModule.h
--- a/Engine/Plugins/Enterprise/DatasmithImporter/Source/DatasmithImporter/Public/DatasmithImporterModule.h
+++ b/Engine/Plugins/Enterprise/DatasmithImporter/Source/DatasmithImporter/Public/DatasmithImporterModule.h
@@ -6,6 +6,7 @@
 #include "Modules/ModuleManager.h"		// For inline LoadModuleChecked()
 #include "DatasmithCore.h"
 
+#include "Delegates/DelegateCombinations.h"
 #include "Templates/Function.h"
 
 #define DATASMITHIMPORTER_MODULE_NAME TEXT("DatasmithImporter")
@@ -42,5 +43,8 @@ public:
 	 * Restores the Datasmith imported values on an object
 	 */
 	virtual void ResetOverrides( UObject* Object ) = 0;
+
+	DECLARE_MULTICAST_DELEGATE_OneParam(FOnGenerateDatasmithImportMenu, struct FToolMenuSection&);
+	virtual FOnGenerateDatasmithImportMenu& OnGenerateDatasmithImportMenu() = 0;
 };
 
