DatasmithImporter - Fix memory leak caused by missing UnloadScene() call in FExternalSource.

Without this call the DatasmithTranslators were never explicitly unloaded, causing memory leak with translators loading external resources like DLLs.

diff --git a/Engine/Plugins/Enterprise/DatasmithImporter/Source/ExternalSource/Private/ExternalSource.cpp b/Engine/Plugins/Enterprise/DatasmithImporter/Source/ExternalSource/Private/ExternalSource.cpp
--- a/Engine/Plugins/Enterprise/DatasmithImporter/Source/ExternalSource/Private/ExternalSource.cpp
+++ b/Engine/Plugins/Enterprise/DatasmithImporter/Source/ExternalSource/Private/ExternalSource.cpp
@@ -25,7 +25,17 @@ namespace UE::DatasmithImporter
 		: SourceUri(InSourceUri)
 	{}
 
-	FExternalSource::~FExternalSource() {}
+	FExternalSource::~FExternalSource()
+	{
+		/*
+		 *  todo: This is should not be done by the FExternalSource but rather by the code loading the scene.
+		 *        Ideally the FExternalSource should provide an interface making use of RAII to manage the loaded state of the Translator.
+		 */
+		if (AssetTranslator)
+		{
+			AssetTranslator->UnloadScene();
+		}
+	}
 
 	const TSharedPtr<IDatasmithTranslator>& FExternalSource::GetAssetTranslator()
 	{
