Exposing FSceneGuard and using it to manage the loaded state of FExternalSource's translator

Also adding missing call to LoadScene() when importing a scene with DirectLink.
This was non-incidental in the previous release as the native translator use to be stateless, but it is no longer the case and the Load/Unload lifecycle requirements must be enforced for proper scene translation.

diff --git a/Engine/Plugins/Enterprise/DatasmithImporter/Source/DatasmithExternalSource/Private/DatasmithFileExternalSource.cpp b/Engine/Plugins/Enterprise/DatasmithImporter/Source/DatasmithExternalSource/Private/DatasmithFileExternalSource.cpp
--- a/Engine/Plugins/Enterprise/DatasmithImporter/Source/DatasmithExternalSource/Private/DatasmithFileExternalSource.cpp
+++ b/Engine/Plugins/Enterprise/DatasmithImporter/Source/DatasmithExternalSource/Private/DatasmithFileExternalSource.cpp
@@ -45,16 +45,8 @@ namespace UE::DatasmithImporter
 
 	TSharedPtr<IDatasmithScene> FDatasmithFileExternalSource::LoadImpl()
 	{
-		const TSharedPtr<IDatasmithTranslator>& Translator = GetAssetTranslator();
-
-		if (!Translator.IsValid())
-		{
-			UE_LOG(LogDatasmithFileExternalSource, Warning, TEXT("Datasmith import error: no suitable translator found for this external source. Abort import."));
-			return nullptr;
-		}
-	
 		TSharedRef<IDatasmithScene> LoadedScene = FDatasmithSceneFactory::CreateScene(*GetSceneName());
-		if (!Translator->LoadScene(LoadedScene))
+		if (!TranslatorLoadScene(LoadedScene))
 		{
 			UE_LOG(LogDatasmithFileExternalSource, Warning, TEXT("Datasmith import error: Scene translation failure. Abort import."));
 			return nullptr;
diff --git a/Engine/Plugins/Enterprise/DatasmithImporter/Source/DatasmithTranslator/Private/DatasmithTranslatableSource.cpp b/Engine/Plugins/Enterprise/DatasmithImporter/Source/DatasmithTranslator/Private/DatasmithTranslatableSource.cpp
--- a/Engine/Plugins/Enterprise/DatasmithImporter/Source/DatasmithTranslator/Private/DatasmithTranslatableSource.cpp
+++ b/Engine/Plugins/Enterprise/DatasmithImporter/Source/DatasmithTranslator/Private/DatasmithTranslatableSource.cpp
@@ -6,27 +6,20 @@
 
 #define LOCTEXT_NAMESPACE "DatasmithTranslatableSceneSource"
 
-class FSceneGuard
+
+FDatasmithSceneGuard::FDatasmithSceneGuard(const TSharedPtr<IDatasmithTranslator>& Translator, const TSharedRef<IDatasmithScene>& Scene, bool& bOutLoadOk)
+	: Translator(Translator)
 {
-public:
-	FSceneGuard(TSharedPtr<IDatasmithTranslator> Translator, TSharedRef< IDatasmithScene > Scene, bool& bOutLoadOk)
-		: Translator(Translator)
-	{
-		bOutLoadOk = Translator.IsValid() && Translator->LoadScene(Scene);
-	}
+	bOutLoadOk = Translator.IsValid() && Translator->LoadScene(Scene);
+}
 
-	~FSceneGuard()
+FDatasmithSceneGuard::~FDatasmithSceneGuard()
+{
+	if (Translator.IsValid())
 	{
-		if (Translator.IsValid())
-		{
-			Translator->UnloadScene();
-		}
+		Translator->UnloadScene();
 	}
-
-private:
-	TSharedPtr<IDatasmithTranslator> Translator;
-};
-
+}
 
 FDatasmithTranslatableSceneSource::FDatasmithTranslatableSceneSource(const FDatasmithSceneSource& Source)
 	: Translator(nullptr)
@@ -52,7 +45,7 @@ bool FDatasmithTranslatableSceneSource::Translate(TSharedRef< IDatasmithScene >
 	bool bIsAlreadyLoaded = SceneGuard.IsValid();
 	if (IsTranslatable() && !bIsAlreadyLoaded)
 	{
-		SceneGuard.Reset(new FSceneGuard(Translator, Scene, bLoadedOk));
+		SceneGuard.Reset(new FDatasmithSceneGuard(Translator, Scene, bLoadedOk));
 	}
 	return bLoadedOk;
 }
diff --git a/Engine/Plugins/Enterprise/DatasmithImporter/Source/DatasmithTranslator/Public/DatasmithTranslatableSource.h b/Engine/Plugins/Enterprise/DatasmithImporter/Source/DatasmithTranslator/Public/DatasmithTranslatableSource.h
--- a/Engine/Plugins/Enterprise/DatasmithImporter/Source/DatasmithTranslator/Public/DatasmithTranslatableSource.h
+++ b/Engine/Plugins/Enterprise/DatasmithImporter/Source/DatasmithTranslator/Public/DatasmithTranslatableSource.h
@@ -7,6 +7,23 @@
 #include "CoreMinimal.h"
 #include "UObject/StrongObjectPtr.h"
 
+class IDatasmithScene;
+class IDatasmithTranslator;
+
+/**
+ * Scopes a translator's loading lifecycle.
+ */
+class DATASMITHTRANSLATOR_API FDatasmithSceneGuard
+{
+public:
+	FDatasmithSceneGuard(const TSharedPtr<IDatasmithTranslator>& Translator, const TSharedRef<IDatasmithScene>& Scene, bool& bOutLoadOk);
+
+	~FDatasmithSceneGuard();
+
+private:
+	TSharedPtr<IDatasmithTranslator> Translator;
+};
+
 /**
  * Wrap a source with an adapted translator.
  * This scopes the lifecycle of a translator, 
@@ -18,17 +35,14 @@ struct DATASMITHTRANSLATOR_API FDatasmithTranslatableSceneSource
 
 	bool IsTranslatable() const;
 
-	bool Translate(TSharedRef< class IDatasmithScene > Scene);
+	bool Translate(TSharedRef<IDatasmithScene> Scene);
 
-	TSharedPtr<class IDatasmithTranslator> GetTranslator() const;
+	TSharedPtr<IDatasmithTranslator> GetTranslator() const;
 
 private:
 	/** Translator currently in use (null when not importing) */
-	TSharedPtr<class IDatasmithTranslator> Translator;
+	TSharedPtr<IDatasmithTranslator> Translator;
 
 	/** internal helper to release scene */
-	TUniquePtr< class FSceneGuard > SceneGuard;
-};
-
-
-
+	TUniquePtr<FDatasmithSceneGuard> SceneGuard;
+};
\ No newline at end of file
diff --git a/Engine/Plugins/Enterprise/DatasmithImporter/Source/DirectLinkExtension/Private/DirectLinkExternalSource.cpp b/Engine/Plugins/Enterprise/DatasmithImporter/Source/DirectLinkExtension/Private/DirectLinkExternalSource.cpp
--- a/Engine/Plugins/Enterprise/DatasmithImporter/Source/DirectLinkExtension/Private/DirectLinkExternalSource.cpp
+++ b/Engine/Plugins/Enterprise/DatasmithImporter/Source/DirectLinkExtension/Private/DirectLinkExternalSource.cpp
@@ -34,6 +34,12 @@ namespace UE::DatasmithImporter
 			SceneReceiver->FinalSnapshot(SceneSnapshot);
 			if (TSharedPtr<FDirectLinkExternalSource> PinnedExternalSource = DirectLinkExternalSource.Pin())
 			{
+				if (TSharedPtr<IDatasmithScene> Scene = PinnedExternalSource->GetDatasmithScene())
+				{
+					// Even though the scene is already loaded, it needs to be loaded inside the translator before we can load additional payload with it.
+					PinnedExternalSource->TranslatorLoadScene(Scene.ToSharedRef());
+				}
+
 				PinnedExternalSource->CachedHash = GenerateSceneSnapshotHash(SceneSnapshot);
 				PinnedExternalSource->TriggerOnExternalSourceChanged();
 			}
diff --git a/Engine/Plugins/Enterprise/DatasmithImporter/Source/ExternalSource/Private/ExternalSource.cpp b/Engine/Plugins/Enterprise/DatasmithImporter/Source/ExternalSource/Private/ExternalSource.cpp
--- a/Engine/Plugins/Enterprise/DatasmithImporter/Source/ExternalSource/Private/ExternalSource.cpp
+++ b/Engine/Plugins/Enterprise/DatasmithImporter/Source/ExternalSource/Private/ExternalSource.cpp
@@ -5,6 +5,7 @@
 #include "Async/Async.h"
 #include "DatasmithSceneFactory.h"
 #include "DatasmithSceneSource.h"
+#include "DatasmithTranslatableSource.h"
 #include "DatasmithTranslator.h"
 #include "DatasmithTranslatorManager.h"
 #include "DatasmithUtils.h"
@@ -26,16 +27,7 @@ namespace UE::DatasmithImporter
 	{}
 
 	FExternalSource::~FExternalSource()
-	{
-		/*
-		 *  todo: This is should not be done by the FExternalSource but rather by the code loading the scene.
-		 *        Ideally the FExternalSource should provide an interface making use of RAII to manage the loaded state of the Translator.
-		 */
-		if (AssetTranslator)
-		{
-			AssetTranslator->UnloadScene();
-		}
-	}
+	{}
 
 	const TSharedPtr<IDatasmithTranslator>& FExternalSource::GetAssetTranslator()
 	{
@@ -234,6 +226,14 @@ namespace UE::DatasmithImporter
 		}
 	}
 
+	bool FExternalSource::TranslatorLoadScene(const TSharedRef<IDatasmithScene>& Scene)
+	{
+		bool bLoaded = false;
+		//Unload any previously loaded scene, and load a new one.
+		SceneGuard = MakeUnique<FDatasmithSceneGuard>(GetAssetTranslator(), Scene, bLoaded);
+		return bLoaded;
+	}
+
 	void FExternalSource::ValidateDatasmithVersion() const
 	{
 		if (TSharedPtr<IDatasmithScene> Scene = GetDatasmithScene())
diff --git a/Engine/Plugins/Enterprise/DatasmithImporter/Source/ExternalSource/Public/ExternalSource.h b/Engine/Plugins/Enterprise/DatasmithImporter/Source/ExternalSource/Public/ExternalSource.h
--- a/Engine/Plugins/Enterprise/DatasmithImporter/Source/ExternalSource/Public/ExternalSource.h
+++ b/Engine/Plugins/Enterprise/DatasmithImporter/Source/ExternalSource/Public/ExternalSource.h
@@ -19,6 +19,7 @@
 #endif //WITH_EDITOR
 
 class FAsyncTaskNotification;
+class FDatasmithSceneGuard;
 class IDatasmithScene;
 class IDatasmithTranslator;
 
@@ -164,6 +165,11 @@ namespace UE::DatasmithImporter
 		 */
 		void TriggerOnExternalSourceChanged();
 
+		/**
+		 * Loads a scene from the translator and starts a new translator load lifecycle.
+		 */
+		bool TranslatorLoadScene(const TSharedRef<IDatasmithScene>& Scene);
+
 	private:
 
 		typedef TOptional<TPromise<TSharedPtr<IDatasmithScene>>> FOptionalScenePromise;
@@ -190,5 +196,7 @@ namespace UE::DatasmithImporter
 		FString SceneName;
 
 		TSharedPtr<IDatasmithTranslator> AssetTranslator;
+
+		TUniquePtr<FDatasmithSceneGuard> SceneGuard;
 	};
 }
\ No newline at end of file
