Datasmith: add support of doted property path for DatasmithKeyValueProperty

diff --git a/Engine/Plugins/Enterprise/DatasmithImporter/Source/DatasmithImporter/DatasmithImporter.Build.cs b/Engine/Plugins/Enterprise/DatasmithImporter/Source/DatasmithImporter/DatasmithImporter.Build.cs
--- a/Engine/Plugins/Enterprise/DatasmithImporter/Source/DatasmithImporter/DatasmithImporter.Build.cs
+++ b/Engine/Plugins/Enterprise/DatasmithImporter/Source/DatasmithImporter/DatasmithImporter.Build.cs
@@ -31,6 +31,7 @@ namespace UnrealBuildTool.Rules
 					"InterchangeFactoryNodes",
 					"InterchangeNodes",
 					"Json",
+					"Kismet",
 					"Landscape",
 					"LandscapeEditor",
 					"LandscapeEditorUtilities",
diff --git a/Engine/Plugins/Enterprise/DatasmithImporter/Source/DatasmithImporter/Private/DatasmithActorImporter.cpp b/Engine/Plugins/Enterprise/DatasmithImporter/Source/DatasmithImporter/Private/DatasmithActorImporter.cpp
--- a/Engine/Plugins/Enterprise/DatasmithImporter/Source/DatasmithImporter/Private/DatasmithActorImporter.cpp
+++ b/Engine/Plugins/Enterprise/DatasmithImporter/Source/DatasmithImporter/Private/DatasmithActorImporter.cpp
@@ -25,6 +25,7 @@
 #include "CineCameraComponent.h"
 #include "Components/DecalComponent.h"
 #include "Components/HierarchicalInstancedStaticMeshComponent.h"
+#include "DiffUtils.h"
 #include "Editor.h"
 #include "Editor/EditorEngine.h"
 #include "Engine/DecalActor.h"
@@ -384,13 +385,40 @@ AActor* FDatasmithActorImporter::ImportCustomActor(FDatasmithImportContext& Impo
 
 	for ( int32 i = 0; i < InCustomActorElement->GetPropertiesCount(); ++i )
 	{
-		TSharedPtr< IDatasmithKeyValueProperty > KeyValueProperty = InCustomActorElement->GetProperty( i );
+		const TSharedPtr< IDatasmithKeyValueProperty >& KeyValueProperty = InCustomActorElement->GetProperty( i );
 
-		FProperty* Property = FindFProperty< FProperty >( Actor->GetClass(), KeyValueProperty->GetName() );
+		FString PropStr = KeyValueProperty->GetName();
 
-		if ( Property )
+		TArray<FString> PropNamesStr;
+		PropStr.ParseIntoArray(PropNamesStr, TEXT("."));
+
+		TArray<FName> PropNames;
+		for (auto&& Str : PropNamesStr)
+		{
+			PropNames.Add(FName(*Str));
+		}
+
+		FPropertySoftPath PropertyPath(PropNames);
+		FResolvedProperty ResolvedProperty = PropertyPath.Resolve(Actor);
+
+		void* Container = const_cast<void*>(ResolvedProperty.Object);
+		if (!ResolvedProperty.Property || !Container)
+		{
+			ImportContext.LogWarning(FText::Format(LOCTEXT("BadPropertyKey", "Cannot find Property '{0}' on actor {1}.")
+				, FText::FromString(PropStr)
+				, FText::FromString(InCustomActorElement->GetName() )) 
+			);
+			continue;
+		}
+
+		const TCHAR* Remaining = ResolvedProperty.Property->ImportText_InContainer( KeyValueProperty->GetValue(), Container, nullptr, PPF_None );
+		if (Remaining && Remaining[0] != 0)
 		{
-			Property->ImportText_InContainer( KeyValueProperty->GetValue(), Actor, nullptr, PPF_None );
+			ImportContext.LogWarning(FText::Format(LOCTEXT("BadPropertyValue", "Cannot properly load value '{0}' on property '{1}' for actor {2}.")
+				, FText::FromStringView(KeyValueProperty->GetValue())
+				, FText::FromString(PropStr)
+				, FText::FromStringView(InCustomActorElement->GetName())) 
+			);
 		}
 	}
 
diff --git a/Engine/Source/Runtime/Datasmith/DatasmithCore/Private/DatasmithSceneElementsImpl.cpp b/Engine/Source/Runtime/Datasmith/DatasmithCore/Private/DatasmithSceneElementsImpl.cpp
--- a/Engine/Source/Runtime/Datasmith/DatasmithCore/Private/DatasmithSceneElementsImpl.cpp
+++ b/Engine/Source/Runtime/Datasmith/DatasmithCore/Private/DatasmithSceneElementsImpl.cpp
@@ -112,8 +112,8 @@ TSharedPtr< IDatasmithKeyValueProperty > FDatasmithKeyValuePropertyImpl::NullPro
 FDatasmithKeyValuePropertyImpl::FDatasmithKeyValuePropertyImpl(const TCHAR* InName)
 	: FDatasmithElementImpl(InName, EDatasmithElementType::KeyValueProperty)
 	, PropertyType(EDatasmithKeyValuePropertyType::String)
-	, Value(InName)
 {
+	FDatasmithKeyValuePropertyImpl::SetName(InName);
 	Store.RegisterParameter(Value, "Value");
 	Store.RegisterParameter(PropertyType, "PropertyType");
 }
diff --git a/Engine/Source/Runtime/Datasmith/DatasmithCore/Private/DatasmithSceneElementsImpl.h b/Engine/Source/Runtime/Datasmith/DatasmithCore/Private/DatasmithSceneElementsImpl.h
--- a/Engine/Source/Runtime/Datasmith/DatasmithCore/Private/DatasmithSceneElementsImpl.h
+++ b/Engine/Source/Runtime/Datasmith/DatasmithCore/Private/DatasmithSceneElementsImpl.h
@@ -72,6 +72,7 @@ class FDatasmithKeyValuePropertyImpl : public FDatasmithElementImpl< IDatasmithK
 {
 public:
 	FDatasmithKeyValuePropertyImpl(const TCHAR* InName);
+	virtual void SetName(const TCHAR* InName) override { Name = InName; }
 
 	EDatasmithKeyValuePropertyType GetPropertyType() const override { return PropertyType; }
 	void SetPropertyType( EDatasmithKeyValuePropertyType InType ) override;
