Datasmith: Moved path resolution logic out of the xml reader in order to share that logic with scenes received from DirectLink

diff --git a/Engine/Plugins/Enterprise/DatasmithImporter/Source/DatasmithExternalSource/Private/DatasmithDirectLinkTranslator.cpp b/Engine/Plugins/Enterprise/DatasmithImporter/Source/DatasmithExternalSource/Private/DatasmithDirectLinkTranslator.cpp
--- a/Engine/Plugins/Enterprise/DatasmithImporter/Source/DatasmithExternalSource/Private/DatasmithDirectLinkTranslator.cpp
+++ b/Engine/Plugins/Enterprise/DatasmithImporter/Source/DatasmithExternalSource/Private/DatasmithDirectLinkTranslator.cpp
@@ -1,6 +1,7 @@
 // Copyright Epic Games, Inc. All Rights Reserved.
 
 #include "DatasmithDirectLinkTranslator.h"
+#include "IDatasmithSceneElements.h"
 
 namespace UE::DatasmithImporter
 {
@@ -9,4 +10,15 @@ namespace UE::DatasmithImporter
 		OutCapabilities.SupportedFileFormats.Emplace(TEXT("directlink"), TEXT("DirectLink stream"));
 		OutCapabilities.bParallelLoadStaticMeshSupported = true;
 	}
+
+
+	bool FDatasmithDirectLinkTranslator::LoadScene(TSharedRef<IDatasmithScene> OutScene)
+	{
+		// With DirectLink the scene is already translated by the SceneReceiver, don't override the scene.
+		FString ResourcePath = OutScene->GetResourcePath();
+		TArray<FString> ResourcePaths;
+		ResourcePath.ParseIntoArray(ResourcePaths, TEXT(";"));
+		ResolveSceneFilePaths(OutScene, ResourcePaths);
+		return true;
+	}
 }
\ No newline at end of file
diff --git a/Engine/Plugins/Enterprise/DatasmithImporter/Source/DatasmithExternalSource/Private/DatasmithDirectLinkTranslator.h b/Engine/Plugins/Enterprise/DatasmithImporter/Source/DatasmithExternalSource/Private/DatasmithDirectLinkTranslator.h
--- a/Engine/Plugins/Enterprise/DatasmithImporter/Source/DatasmithExternalSource/Private/DatasmithDirectLinkTranslator.h
+++ b/Engine/Plugins/Enterprise/DatasmithImporter/Source/DatasmithExternalSource/Private/DatasmithDirectLinkTranslator.h
@@ -9,8 +9,8 @@ namespace UE::DatasmithImporter
 {
 	/**
 	 * Due to the way the Editor's import/reimport utilities currently work, we must provide a path to an existing file.
-	 * The FDatasmithDirecTLinkTranslator is a workaround used to be directed to the DatasmithImportFactory when the 
-	 * source file has a ".directlink" extention. In that case, the actual scene is loaded using the FExternalSource 
+	 * The FDatasmithDirecTLinkTranslator is a workaround used to be directed to the DatasmithImportFactory when the
+	 * source file has a ".directlink" extention. In that case, the actual scene is loaded using the FExternalSource
 	 * generated from the directlink URI.
 	 */
 	class FDatasmithDirectLinkTranslator : public FDatasmithNativeTranslator
@@ -20,10 +20,6 @@ namespace UE::DatasmithImporter
 
 		virtual void Initialize(FDatasmithTranslatorCapabilities& OutCapabilities) override;
 
-		virtual bool LoadScene(TSharedRef<IDatasmithScene> OutScene) override
-		{
-			// With DirectLink the scene is already translated by the SceneReceiver, don't override the scene.
-			return true;
-		}
+		virtual bool LoadScene(TSharedRef<IDatasmithScene> OutScene) override;
 	};
 }
\ No newline at end of file
diff --git a/Engine/Plugins/Enterprise/DatasmithImporter/Source/DatasmithNativeTranslator/Private/DatasmithNativeTranslator.cpp b/Engine/Plugins/Enterprise/DatasmithImporter/Source/DatasmithNativeTranslator/Private/DatasmithNativeTranslator.cpp
--- a/Engine/Plugins/Enterprise/DatasmithImporter/Source/DatasmithNativeTranslator/Private/DatasmithNativeTranslator.cpp
+++ b/Engine/Plugins/Enterprise/DatasmithImporter/Source/DatasmithNativeTranslator/Private/DatasmithNativeTranslator.cpp
@@ -19,6 +19,78 @@
 #include "UObject/StrongObjectPtr.h"
 
 
+FString FDatasmithNativeTranslator::ResolveFilePath(const FString& FilePath, const TArray<FString>& ResourcePaths)
+{
+	if (FPaths::IsRelative(FilePath) && !FPaths::FileExists(FilePath))
+	{
+		for (const FString& ResourcePath : ResourcePaths)
+		{
+			if (ResourcePath.IsEmpty())
+			{
+				continue;
+			}
+
+			FString NewFilePath = ResourcePath / FilePath;
+			if (FPaths::FileExists(NewFilePath))
+			{
+				return NewFilePath;
+			}
+		}
+	}
+	return FilePath;
+}
+
+void FDatasmithNativeTranslator::ResolveSceneFilePaths(TSharedRef<IDatasmithScene> Scene, const TArray<FString>& ResourcePaths)
+{
+	for (int32 Index = 0; Index < Scene->GetMeshesCount(); ++Index)
+	{
+		const TSharedPtr<IDatasmithMeshElement>& Mesh = Scene->GetMesh(Index);
+		const TCHAR* Path = Mesh->GetFile();
+		Mesh->SetFile(*ResolveFilePath(Path, ResourcePaths));
+	}
+
+	for (int32 Index = 0; Index < Scene->GetTexturesCount(); ++Index)
+	{
+		const TSharedPtr<IDatasmithTextureElement>& Tex = Scene->GetTexture(Index);
+		const TCHAR* Path = Tex->GetFile();
+		Tex->SetFile(*ResolveFilePath(Path, ResourcePaths));
+	}
+
+	for (int32 Index = 0; Index < Scene->GetLevelSequencesCount(); ++Index)
+	{
+		const TSharedPtr<IDatasmithLevelSequenceElement>& Sequence = Scene->GetLevelSequence(Index);
+		const TCHAR* Path = Sequence->GetFile();
+		Sequence->SetFile(*ResolveFilePath(Path, ResourcePaths));
+	}
+
+	TFunction<void (const TSharedPtr<IDatasmithActorElement> Actor)> VisitActorTree;
+	VisitActorTree = [&](const TSharedPtr<IDatasmithActorElement> Actor) -> void
+	{
+		if (Actor->IsA(EDatasmithElementType::Landscape))
+		{
+			const TSharedPtr<IDatasmithLandscapeElement>& LandscapeActor = StaticCastSharedPtr<IDatasmithLandscapeElement>(Actor);
+			const TCHAR* Path = LandscapeActor->GetHeightmap();
+			LandscapeActor->SetHeightmap(*ResolveFilePath(Path, ResourcePaths));
+		}
+		else if (Actor->IsA(EDatasmithElementType::Light))
+		{
+			const TSharedPtr<IDatasmithLightActorElement>& Light = StaticCastSharedPtr<IDatasmithLightActorElement>(Actor);
+			const TCHAR* Path = Light->GetIesFile();
+			Light->SetIesFile(*ResolveFilePath(Path, ResourcePaths));
+		}
+
+		for (int32 ChildIndex = 0; ChildIndex < Actor->GetChildrenCount(); ++ChildIndex)
+		{
+			VisitActorTree(Actor->GetChild(ChildIndex));
+		}
+	};
+
+	for (int32 Index = 0; Index < Scene->GetActorsCount(); ++Index)
+	{
+		VisitActorTree(Scene->GetActor(Index));
+	}
+}
+
 void FDatasmithNativeTranslator::Initialize(FDatasmithTranslatorCapabilities& OutCapabilities)
 {
 	OutCapabilities.SupportedFileFormats.Emplace(TEXT("udatasmith"), TEXT("Datasmith files"));
@@ -28,9 +100,22 @@ void FDatasmithNativeTranslator::Initialize(FDatasmithTranslatorCapabilities& Ou
 bool FDatasmithNativeTranslator::LoadScene(TSharedRef<IDatasmithScene> OutScene)
 {
 	FDatasmithSceneXmlReader XmlParser;
-	return XmlParser.ParseFile(GetSource().GetSourceFile(), OutScene);
-}
 
+	bool bParsingResult = XmlParser.ParseFile(GetSource().GetSourceFile(), OutScene);
+	if (bParsingResult)
+	{
+		TArray<FString> ResourcePaths;
+		FString ResourcePath = OutScene->GetResourcePath();
+		ResourcePath.ParseIntoArray(ResourcePaths, TEXT(";"));
+
+		FString ProjectPath = FPaths::GetPath(GetSource().GetSourceFile());
+		ResourcePaths.Insert(ProjectPath, 0);
+
+		ResolveSceneFilePaths(OutScene, ResourcePaths);
+	}
+
+	return bParsingResult;
+}
 
 namespace DatasmithNativeTranslatorImpl
 {
@@ -40,11 +125,11 @@ namespace DatasmithNativeTranslatorImpl
 		TArray< FDatasmithMeshSourceModel > SourceModels;
 	};
 
-	TArray< FDatasmithMeshInternal > GetDatasmithMeshFromMeshElement( const TSharedRef< IDatasmithMeshElement > MeshElement )
+	TArray< FDatasmithMeshInternal > GetDatasmithMeshFromMeshPath( const TCHAR* MeshPath )
 	{
 		TArray< FDatasmithMeshInternal > Result;
 
-		TUniquePtr<FArchive> Archive( IFileManager::Get().CreateFileReader(MeshElement->GetFile()) );
+		TUniquePtr<FArchive> Archive( IFileManager::Get().CreateFileReader(MeshPath) );
 		if ( !Archive.IsValid() )
 		{
 			return Result;
@@ -141,13 +226,14 @@ bool FDatasmithNativeTranslator::LoadStaticMesh(const TSharedRef<IDatasmithMeshE
 
 	using namespace DatasmithNativeTranslatorImpl;
 
-	if (MeshElement->GetFile() == nullptr || !FPaths::FileExists( MeshElement->GetFile() ))
+	FString FilePath = MeshElement->GetFile();
+	if (!FPaths::FileExists(FilePath))
 	{
 		return false;
 	}
 
 	int32 ExtractionFailure = 0;
-	for (FDatasmithMeshInternal& DatasmithMesh : GetDatasmithMeshFromMeshElement( MeshElement ))
+	for (FDatasmithMeshInternal& DatasmithMesh : GetDatasmithMeshFromMeshPath( *FilePath ))
 	{
 		if (DatasmithMesh.bIsCollisionMesh)
 		{
diff --git a/Engine/Plugins/Enterprise/DatasmithImporter/Source/DatasmithNativeTranslator/Public/DatasmithNativeTranslator.h b/Engine/Plugins/Enterprise/DatasmithImporter/Source/DatasmithNativeTranslator/Public/DatasmithNativeTranslator.h
--- a/Engine/Plugins/Enterprise/DatasmithImporter/Source/DatasmithNativeTranslator/Public/DatasmithNativeTranslator.h
+++ b/Engine/Plugins/Enterprise/DatasmithImporter/Source/DatasmithNativeTranslator/Public/DatasmithNativeTranslator.h
@@ -17,5 +17,9 @@ public:
 	virtual bool LoadStaticMesh(const TSharedRef<IDatasmithMeshElement> MeshElement, FDatasmithMeshElementPayload& OutMeshPayload) override;
 
 	virtual bool LoadLevelSequence(const TSharedRef<IDatasmithLevelSequenceElement> LevelSequenceElement, FDatasmithLevelSequencePayload& OutLevelSequencePayload) override;
+
+protected:
+	static FString ResolveFilePath(const FString& FilePath, const TArray<FString>& ResourcePaths);
+	static void ResolveSceneFilePaths(TSharedRef<IDatasmithScene> Scene, const TArray<FString>& ResourcePaths);
 };
 
diff --git a/Engine/Source/Runtime/Datasmith/DatasmithCore/Private/DatasmithSceneXmlReader.cpp b/Engine/Source/Runtime/Datasmith/DatasmithCore/Private/DatasmithSceneXmlReader.cpp
--- a/Engine/Source/Runtime/Datasmith/DatasmithCore/Private/DatasmithSceneXmlReader.cpp
+++ b/Engine/Source/Runtime/Datasmith/DatasmithCore/Private/DatasmithSceneXmlReader.cpp
@@ -104,24 +104,6 @@ T FDatasmithSceneXmlReader::ValueFromString( const FString& InString ) const
 	return DatasmithSceneXmlReaderImpl::ValueFromString<T>(InString);
 }
 
-FString FDatasmithSceneXmlReader::ResolveFilePath(const FString& AssetFile) const
-{
-	if ( !FPaths::IsRelative( AssetFile ) )
-	{
-		return AssetFile;
-	}
-
-	FString FullAssetPath = FPaths::Combine(ProjectPath, AssetFile);
-
-	if ( FPaths::FileExists(FullAssetPath) )
-	{
-		return FullAssetPath;
-	}
-	else
-	{
-		return AssetFile;
-	}
-}
 
 void FDatasmithSceneXmlReader::PatchUpVersion(TSharedRef< IDatasmithScene >& OutScene)
 {
@@ -166,7 +148,7 @@ void FDatasmithSceneXmlReader::ParseLevelSequence(FXmlNode* InNode, const TShare
 	{
 		if (ChildrenNodes[j]->GetTag() == TEXT("file"))
 		{
-			OutElement->SetFile( *ResolveFilePath(ChildrenNodes[j]->GetAttribute(TEXT("path"))) );
+			OutElement->SetFile( *ChildrenNodes[j]->GetAttribute(TEXT("path")) );
 		}
 		else if (ChildrenNodes[j]->GetTag() == DATASMITH_HASH)
 		{
@@ -366,7 +348,7 @@ void FDatasmithSceneXmlReader::ParseMesh(FXmlNode* InNode, TSharedPtr<IDatasmith
 	{
 		if (MeshNodes[j]->GetTag() == TEXT("file"))
 		{
-			OutElement->SetFile( *ResolveFilePath(MeshNodes[j]->GetAttribute(TEXT("path"))) );
+			OutElement->SetFile(*MeshNodes[j]->GetAttribute(TEXT("path"))) ;
 		}
 		else if (MeshNodes[j]->GetTag() == TEXT("Size"))
 		{
@@ -438,7 +420,7 @@ void FDatasmithSceneXmlReader::ParseTextureElement(FXmlNode* InNode, TSharedPtr<
 	{
 		OutElement->SetTextureAddressY((EDatasmithTextureAddress)ValueFromString<int32>(StrValue));
 	}
-	OutElement->SetFile(*ResolveFilePath(InNode->GetAttribute(TEXT("file"))));
+	OutElement->SetFile(*InNode->GetAttribute(TEXT("file")));
 
 	const TArray<FXmlNode*>& TexNode = InNode->GetChildrenNodes();
 	for (int i = 0; i < TexNode.Num(); ++i)
@@ -911,7 +893,7 @@ void FDatasmithSceneXmlReader::ParseLight(FXmlNode* InNode, TSharedPtr<IDatasmit
 		else if (ChildNode->GetTag() == DATASMITH_LIGHTIESNAME)
 		{
 			OutElement->SetUseIes(true);
-			OutElement->SetIesFile( *ResolveFilePath(ChildNode->GetAttribute(TEXT("file"))) );
+			OutElement->SetIesFile( *ChildNode->GetAttribute(TEXT("file")) );
 		}
 		else if (ChildNode->GetTag() == DATASMITH_LIGHTIESTEXTURENAME)
 		{
@@ -2075,7 +2057,7 @@ void FDatasmithSceneXmlReader::ParseLandscape(FXmlNode* InNode, TSharedRef< IDat
 	{
 		if ( ChildNode->GetTag() == DATASMITH_HEIGHTMAPNAME )
 		{
-			OutElement->SetHeightmap( *ResolveFilePath( ChildNode->GetAttribute( TEXT("value") ) ) );
+			OutElement->SetHeightmap( *ChildNode->GetAttribute( TEXT("value") ) );
 		}
 		else if ( ChildNode->GetTag() == DATASMITH_MATERIAL )
 		{
diff --git a/Engine/Source/Runtime/Datasmith/DatasmithCore/Public/IDatasmithSceneElements.h b/Engine/Source/Runtime/Datasmith/DatasmithCore/Public/IDatasmithSceneElements.h
--- a/Engine/Source/Runtime/Datasmith/DatasmithCore/Public/IDatasmithSceneElements.h
+++ b/Engine/Source/Runtime/Datasmith/DatasmithCore/Public/IDatasmithSceneElements.h
@@ -1505,16 +1505,14 @@ public:
 	 */
 	virtual void SetProductVersion(const TCHAR*) = 0;
 
-	/** Returns the original path resources were stored */
+	/** Returns the ';' separated list of paths where resources are stored */
 	virtual const TCHAR* GetResourcePath() const = 0;
 
 	/**
-	 * Sets the original path resources were stored.
+	 * Similar to how the PATH environment variable works, sets list of paths where resources can be stored.
 	 *
-	 * @param InResoucePath	The original path
+	 * @param InResoucePath	The ';' separated list of paths
 	 */
-	// #ue_directlink_design: Find a better way to allow DirectLink clients to retrieve assets
-	// associated with elements. Assets will have to be passed over network too
 	virtual void SetResourcePath(const TCHAR*) = 0;
 
 	/** Returns the user identifier who exported the scene */
