DatasmithSDK - Add missing xml serialization of PbrMaterial ClearCoat properties.

diff --git a/Engine/Plugins/Enterprise/DatasmithImporter/Source/DatasmithImporter/Private/DatasmithMaterialImporter.cpp b/Engine/Plugins/Enterprise/DatasmithImporter/Source/DatasmithImporter/Private/DatasmithMaterialImporter.cpp
--- a/Engine/Plugins/Enterprise/DatasmithImporter/Source/DatasmithImporter/Private/DatasmithMaterialImporter.cpp
+++ b/Engine/Plugins/Enterprise/DatasmithImporter/Source/DatasmithImporter/Private/DatasmithMaterialImporter.cpp
@@ -149,6 +149,8 @@ namespace DatasmithMaterialImporterUtils
 		Hash = HashCombine( Hash, ComputeExpressionInputHash( &MaterialElement->GetNormal(), VisitedExpressions ) );
 		Hash = HashCombine( Hash, ComputeExpressionInputHash( &MaterialElement->GetRefraction(), VisitedExpressions ) );
 		Hash = HashCombine( Hash, ComputeExpressionInputHash( &MaterialElement->GetAmbientOcclusion(), VisitedExpressions ) );
+		Hash = HashCombine( Hash, ComputeExpressionInputHash( &MaterialElement->GetClearCoat(), VisitedExpressions ) );
+		Hash = HashCombine( Hash, ComputeExpressionInputHash( &MaterialElement->GetClearCoatRoughness(), VisitedExpressions ) );
 		Hash = HashCombine( Hash, ComputeExpressionInputHash( &MaterialElement->GetMaterialAttributes(), VisitedExpressions ) );
 
 		TSet<IDatasmithMaterialExpression*> AllExpressions;
diff --git a/Engine/Source/Runtime/Datasmith/DatasmithCore/Private/DatasmithSceneXmlReader.cpp b/Engine/Source/Runtime/Datasmith/DatasmithCore/Private/DatasmithSceneXmlReader.cpp
--- a/Engine/Source/Runtime/Datasmith/DatasmithCore/Private/DatasmithSceneXmlReader.cpp
+++ b/Engine/Source/Runtime/Datasmith/DatasmithCore/Private/DatasmithSceneXmlReader.cpp
@@ -1910,6 +1910,8 @@ void FDatasmithSceneXmlReader::ParseUEPbrMaterial(FXmlNode* InNode, TSharedPtr<
 	TryConnectMaterialInput(OutElement->GetNormal());
 	TryConnectMaterialInput(OutElement->GetRefraction());
 	TryConnectMaterialInput(OutElement->GetAmbientOcclusion());
+	TryConnectMaterialInput(OutElement->GetClearCoat());
+	TryConnectMaterialInput(OutElement->GetClearCoatRoughness());
 	TryConnectMaterialInput(OutElement->GetMaterialAttributes());
 
 	for ( const FXmlNode* ChildNode : InNode->GetChildrenNodes() )
diff --git a/Engine/Source/Runtime/Datasmith/DatasmithCore/Private/DatasmithSceneXmlWriter.cpp b/Engine/Source/Runtime/Datasmith/DatasmithCore/Private/DatasmithSceneXmlWriter.cpp
--- a/Engine/Source/Runtime/Datasmith/DatasmithCore/Private/DatasmithSceneXmlWriter.cpp
+++ b/Engine/Source/Runtime/Datasmith/DatasmithCore/Private/DatasmithSceneXmlWriter.cpp
@@ -1422,6 +1422,8 @@ void FDatasmithSceneXmlWriterImpl::WriteUEPbrMaterialElement( const TSharedRef<
 	WriteUEPbrMaterialExpressionInput( MaterialElement, MaterialElement->GetNormal(), Archive, Indent + 1 );
 	WriteUEPbrMaterialExpressionInput( MaterialElement, MaterialElement->GetRefraction(), Archive, Indent + 1 );
 	WriteUEPbrMaterialExpressionInput( MaterialElement, MaterialElement->GetAmbientOcclusion(), Archive, Indent + 1 );
+	WriteUEPbrMaterialExpressionInput( MaterialElement, MaterialElement->GetClearCoat(), Archive, Indent + 1 );
+	WriteUEPbrMaterialExpressionInput( MaterialElement, MaterialElement->GetClearCoatRoughness(), Archive, Indent + 1 );
 	WriteUEPbrMaterialExpressionInput( MaterialElement, MaterialElement->GetMaterialAttributes(), Archive, Indent + 1 );
 
 	if ( MaterialElement->GetUseMaterialAttributes() )
