DirectLink in Editor - Fixing deadlock caused by recursive read-lock.

When a write-lock is requested, additional read-locks will queue up behind the pending write-lock, even if the FRWLock is currently locked in a read (shared) state. This can cause deadlock when we attempt to read-lock the FRWLock recurisivelly.

In this particular case a deadlock could occur between FDirectLinkManager::GetUriFromSourceHandle() and FDirectLinkManager::OnStateChanged(). To fix this now free the read-lock in FDirectLinkManager::GetOrCreateExternalSource() before calling GetUriFromSourceHandle() to avoid locking recursively.

diff --git a/Engine/Plugins/Enterprise/DatasmithImporter/Source/DirectLinkExtension/Private/DirectLinkManager.cpp b/Engine/Plugins/Enterprise/DatasmithImporter/Source/DirectLinkExtension/Private/DirectLinkManager.cpp
--- a/Engine/Plugins/Enterprise/DatasmithImporter/Source/DirectLinkExtension/Private/DirectLinkManager.cpp
+++ b/Engine/Plugins/Enterprise/DatasmithImporter/Source/DirectLinkExtension/Private/DirectLinkManager.cpp
@@ -172,8 +172,11 @@ namespace UE::DatasmithImporter
 		}
 		else if (RegisteredExternalSourcesInfo.Num() > 0)
 		{
-			FRWScopeLock ScopeLock(RawInfoLock, FRWScopeLockType::SLT_ReadOnly);
-			DirectLink::FRawInfo::FDataPointInfo* SourceDataPointInfo = RawInfoCache.DataPointsInfo.Find(SourceHandle);
+			DirectLink::FRawInfo::FDataPointInfo* SourceDataPointInfo;
+			{
+				FRWScopeLock ScopeLock(RawInfoLock, FRWScopeLockType::SLT_ReadOnly);
+				SourceDataPointInfo = RawInfoCache.DataPointsInfo.Find(SourceHandle);
+			}
 			FSourceUri ExternalSourceUri(GetUriFromSourceHandle(SourceHandle));
 
 			if (SourceDataPointInfo && ExternalSourceUri.IsValid())
