Fixing refresh button on DirectLink source selection window

#rn JeanLuc.Corenthin

diff --git a/Engine/Plugins/Enterprise/DatasmithImporter/Source/DirectLinkExtension/Private/UI/SDirectLinkAvailableSource.cpp b/Engine/Plugins/Enterprise/DatasmithImporter/Source/DirectLinkExtension/Private/UI/SDirectLinkAvailableSource.cpp
--- a/Engine/Plugins/Enterprise/DatasmithImporter/Source/DirectLinkExtension/Private/UI/SDirectLinkAvailableSource.cpp
+++ b/Engine/Plugins/Enterprise/DatasmithImporter/Source/DirectLinkExtension/Private/UI/SDirectLinkAvailableSource.cpp
@@ -133,7 +133,7 @@ namespace UE::DatasmithImporter
 			.Padding(2.f)
 			.VAlign(VAlign_Fill)
 			[
-				SNew(SListView< TSharedRef<FDirectLinkExternalSourceInfo>>)
+				SAssignNew(SourceListView, SListView< TSharedRef<FDirectLinkExternalSourceInfo>>)
 				.Visibility(this, &SDirectLinkAvailableSource::GetConnectionViewVisibility)
 				.ListItemsSource(&DirectLinkExternalSourceInfos)
 				.OnGenerateRow(this, &SDirectLinkAvailableSource::OnGenerateRow)
@@ -156,7 +156,7 @@ namespace UE::DatasmithImporter
 					.ButtonStyle(FAppStyle::Get(), "SimpleButton")
 					.ContentPadding(1)
 					.ToolTipText(LOCTEXT("RefreshSourceListButtonTooltip", "Refresh the list of Direct Link sources."))
-					.OnClicked_Lambda([this]() { GenerateDirectLinkExternalSourceInfos(); return FReply::Handled(); })
+					.OnClicked_Lambda([this]() { GenerateDirectLinkExternalSourceInfos(); SourceListView->RequestListRefresh(); return FReply::Handled(); })
 					[
 						SNew(SImage)
 						.Image(FEditorStyle::GetBrush("Icons.Refresh"))
diff --git a/Engine/Plugins/Enterprise/DatasmithImporter/Source/DirectLinkExtension/Private/UI/SDirectLinkAvailableSource.h b/Engine/Plugins/Enterprise/DatasmithImporter/Source/DirectLinkExtension/Private/UI/SDirectLinkAvailableSource.h
--- a/Engine/Plugins/Enterprise/DatasmithImporter/Source/DirectLinkExtension/Private/UI/SDirectLinkAvailableSource.h
+++ b/Engine/Plugins/Enterprise/DatasmithImporter/Source/DirectLinkExtension/Private/UI/SDirectLinkAvailableSource.h
@@ -84,6 +84,8 @@ namespace UE::DatasmithImporter
 
 		TSharedPtr<FDirectLinkExternalSource> SelectedSource;
 
+		TSharedPtr<SListView<TSharedRef<struct FDirectLinkExternalSourceInfo>>> SourceListView;
+
 		TArray<TSharedRef<struct FDirectLinkExternalSourceInfo>> DirectLinkExternalSourceInfos;
 	};
 }
\ No newline at end of file
