BuildFix NoPCH Compile LyraEditor Win64 - SDirectLinkAvailableSource.h error: 'SListView': undeclared identifier

diff --git a/Engine/Plugins/Enterprise/DatasmithImporter/Source/DirectLinkExtension/Private/UI/SDirectLinkAvailableSource.h b/Engine/Plugins/Enterprise/DatasmithImporter/Source/DirectLinkExtension/Private/UI/SDirectLinkAvailableSource.h
--- a/Engine/Plugins/Enterprise/DatasmithImporter/Source/DirectLinkExtension/Private/UI/SDirectLinkAvailableSource.h
+++ b/Engine/Plugins/Enterprise/DatasmithImporter/Source/DirectLinkExtension/Private/UI/SDirectLinkAvailableSource.h
@@ -6,6 +6,7 @@
 
 #include "Framework/Application/SlateApplication.h"
 #include "Widgets/SCompoundWidget.h"
+#include "Widgets/Views/SListView.h"
 
 namespace UE::DatasmithImporter
 {
@@ -16,7 +17,7 @@ namespace UE::DatasmithImporter
 	{
 	public:
 
-		SLATE_BEGIN_ARGS(SDirectLinkAvailableSource) 
+		SLATE_BEGIN_ARGS(SDirectLinkAvailableSource)
 		{}
 			SLATE_ARGUMENT(TSharedPtr<SWindow>, WidgetWindow)
 			SLATE_ARGUMENT(FText, ProceedButtonLabel)
@@ -44,7 +45,7 @@ namespace UE::DatasmithImporter
 		TSharedPtr<FDirectLinkExternalSource> GetSelectedSource() const { return SelectedSource; }
 
 	private:
-	
+
 		void GenerateDirectLinkExternalSourceInfos();
 
 		TSharedRef<ITableRow> OnGenerateRow(TSharedRef<struct FDirectLinkExternalSourceInfo> Item, const TSharedRef<STableViewBase>& Owner) const;
