DirectLink in Editor - Refactor of the DirectLink manager to extract the auto-reimport feature from it and put it in an external class FAutoReimportManager class.

This FAutoReimportManager has a much cleaner and generic interface and its implementation does not rely on async tasks anymore to trigger reimport on mainthread.
Instead it is a simple FTickableEditorObject that process requested reimports during tick.
Some changes are still required on the FExternalSource side to properly interface with the auto-reimport capability in order to move this into the ExternalSource module.

diff --git a/Engine/Plugins/Enterprise/DatasmithImporter/Source/DirectLinkExtension/Private/AutoReimportManager.cpp b/Engine/Plugins/Enterprise/DatasmithImporter/Source/DirectLinkExtension/Private/AutoReimportManager.cpp
new file mode 100644
--- /dev/null
+++ b/Engine/Plugins/Enterprise/DatasmithImporter/Source/DirectLinkExtension/Private/AutoReimportManager.cpp
@@ -0,0 +1,246 @@
+// Copyright Epic Games, Inc. All Rights Reserved.
+
+#include "AutoReimportManager.h"
+
+#include "Async/Async.h"
+#include "ExternalSource.h"
+#include "SourceUri.h"
+#include "EditorReimportHandler.h"
+#include "ExternalSourceModule.h"
+#include "Stats/Stats.h"
+
+DEFINE_LOG_CATEGORY(LogAutoReimportManager);
+
+namespace UE::DatasmithImporter
+{
+	struct FAutoReimportInfo
+	{
+		FAutoReimportInfo(UObject* InTargetObject, const TSharedRef<FExternalSource>& InExternalSource, FDelegateHandle InImportDelegateHandle)
+			: TargetObject(InTargetObject)
+			, ExternalSource(InExternalSource)
+			, ImportDelegateHandle(InImportDelegateHandle)
+			, bChangedDuringPIE(false)
+		{}
+
+		TSoftObjectPtr<UObject> TargetObject;
+		TSharedRef<FExternalSource> ExternalSource;
+		FDelegateHandle ImportDelegateHandle;
+		bool bChangedDuringPIE;
+	};
+
+	void FAutoReimportManager::Tick(float DeltaTime)
+	{
+		if (!CanTriggerReimport())
+		{
+			return;
+		}
+
+		// Trigger Reimport for modified ExternalSources.
+		{
+			// Using a TSet here so that multiple reimport request for the same external source will only be processed once. 
+			// Doing this allows us to skip redundant reimports, as the reimport already uses the latest data from the ExternalSource.
+			TSet<TSharedRef<FExternalSource>> PendingReimportSet;
+			TSharedPtr<FExternalSource> EnqueuedSourceToReimport;
+			while (PendingReimportQueue.Dequeue(EnqueuedSourceToReimport))
+			{
+				PendingReimportSet.Add(EnqueuedSourceToReimport.ToSharedRef());
+			}
+
+			for (const TSharedRef<FExternalSource>& ExternalSourceToReimport : PendingReimportSet)
+			{
+				TriggerAutoReimportOnExternalSource(ExternalSourceToReimport);
+			}
+		}
+
+		TSharedPtr<FExternalSource> InvalidedSourcePtr;
+		while (PendingInvalidateQueue.Dequeue(InvalidedSourcePtr))
+		{
+			TSharedRef<FExternalSource> InvalidedSource(InvalidedSourcePtr.ToSharedRef());
+			TArray<TSharedRef<FAutoReimportInfo>> AutoReimportInfoList;
+			AutoReimportExternalSourcesMap.MultiFind(InvalidedSource, AutoReimportInfoList);
+			AutoReimportExternalSourcesMap.Remove(InvalidedSource);
+
+			for (const TSharedRef<FAutoReimportInfo>& AutoReimportInfo : AutoReimportInfoList)
+			{
+				UObject* Asset = AutoReimportInfo->TargetObject.Get();
+
+				// Assets that were registered for auto-reimport are unregistered from the source
+				// and put back in the "PendingRegistration" list. That way if a compatible source appear they will register to it.
+				AutoReimportObjectsMap.Remove(Asset);
+				PendingAutoReimportObjects.Add(Asset);
+			}
+		}
+
+		// Try to re-enable auto-reimport for assets who have lost their connection.
+		{
+			TArray<TSoftObjectPtr<UObject>> AutoReimportObjectListCopy(PendingAutoReimportObjects.Array());
+			for (const TSoftObjectPtr<UObject>& Object : AutoReimportObjectListCopy)
+			{
+				if (Object.IsValid())
+				{
+					EnableAssetAutoReimport(Object.Get());
+				}
+			}
+		}
+	}
+
+	TStatId FAutoReimportManager::GetStatId() const
+	{
+		RETURN_QUICK_DECLARE_CYCLE_STAT(FAutoReimportManager, STATGROUP_Tickables);
+	}
+
+
+	bool FAutoReimportManager::CanTriggerReimport() const
+	{
+		if (!ensure(IsInGameThread())
+			|| UE::IsSavingPackage(nullptr)
+			|| IsGarbageCollecting()
+			|| (GIsEditor && FApp::IsGame())) // Is in PIE mode
+		{
+			return false;
+		}
+
+		return true;
+	}
+
+	bool FAutoReimportManager::IsAssetAutoReimportEnabled(UObject* InAsset) const
+	{
+		return AutoReimportObjectsMap.Contains(InAsset) || PendingAutoReimportObjects.Contains(InAsset);
+	}
+
+	bool FAutoReimportManager::SetAssetAutoReimport(UObject* InAsset, bool bEnabled)
+	{
+		return bEnabled ? EnableAssetAutoReimport(InAsset) : DisableAssetAutoReimport(InAsset);
+	}
+
+	void FAutoReimportManager::OnExternalSourceInvalidated(const TSharedRef<FExternalSource>& ExternalSource)
+	{
+		PendingInvalidateQueue.Enqueue(ExternalSource);
+	}
+
+	bool FAutoReimportManager::EnableAssetAutoReimport(UObject* InAsset)
+	{
+		// Enable auto-reimport for InAsset.
+		FAssetData AssetData(InAsset);
+		const FSourceUri Uri = FSourceUri::FromAssetData(AssetData);
+
+		if (Uri.IsValid() && !AutoReimportObjectsMap.Contains(InAsset))
+		{
+			if (TSharedPtr<FExternalSource> ExternalSource = IExternalSourceModule::GetOrCreateExternalSource(Uri))
+			{
+				// Register a delegate triggering a reimport task on the external source snapshotupdate event.
+				// That way the asset will be auto-reimported and kept up-to-date.
+				FDelegateHandle DelegateHandle = ExternalSource->OnExternalSourceChanged.AddRaw(this, &FAutoReimportManager::OnExternalSourceChanged);
+
+				TSharedRef<FExternalSource> ExternalSourceRef = ExternalSource.ToSharedRef();
+				TSharedRef<FAutoReimportInfo> AutoReimportInfo = MakeShared<FAutoReimportInfo>(InAsset, ExternalSourceRef, DelegateHandle);
+
+				AutoReimportObjectsMap.Add(InAsset, AutoReimportInfo);
+				AutoReimportExternalSourcesMap.Add(ExternalSourceRef, AutoReimportInfo);
+				//TODO, we can't just async load here, we must add some sort of "auto-reimport capability" to the external source.
+				ExternalSource->AsyncLoad();
+				PendingAutoReimportObjects.Remove(InAsset);
+				return true;
+			}
+			else
+			{
+				bool bIsAlreadySet = false;
+				PendingAutoReimportObjects.Add(InAsset, &bIsAlreadySet);
+				if (!bIsAlreadySet)
+				{
+					return true;
+				}
+			}
+		}
+
+		return false;
+	}
+
+	bool FAutoReimportManager::DisableAssetAutoReimport(UObject* InAsset)
+	{
+		// Disable auto-reimport for InAsset.
+		if (const TSharedRef<FAutoReimportInfo>* AutoReimportInfoPtr = AutoReimportObjectsMap.Find(InAsset))
+		{
+			// Holding a local reference to the FAutoReimportInfo to ensure its lifetime while we are cleaning up.
+			const TSharedRef<FAutoReimportInfo> AutoReimportInfo(*AutoReimportInfoPtr);
+
+			AutoReimportInfo->ExternalSource->OnExternalSourceChanged.Remove(AutoReimportInfo->ImportDelegateHandle);
+			AutoReimportObjectsMap.Remove(InAsset);
+			AutoReimportExternalSourcesMap.RemoveSingle(AutoReimportInfo->ExternalSource, AutoReimportInfo);
+		}
+		else
+		{
+			PendingAutoReimportObjects.Remove(InAsset);
+		}
+
+		return true;
+	}
+
+	void FAutoReimportManager::UpdateModifiedRegisteredAsset(UObject* InAsset)
+	{
+		if (!IsAssetAutoReimportEnabled(InAsset))
+		{
+			// Asset is not registered, nothing to update.
+			return;
+		}
+
+		const FAssetData AssetData(InAsset);
+		const FSourceUri Uri = FSourceUri::FromAssetData(AssetData);
+		const TSharedPtr<FExternalSource> UpdatedExternalSource = Uri.IsValid() ? IExternalSourceModule::GetOrCreateExternalSource(Uri) : nullptr;
+		if (!UpdatedExternalSource)
+		{
+			// Asset was registered for auto reimport but no longer has a DirectLink source, disable auto reimport.
+			DisableAssetAutoReimport(InAsset);
+			return;
+		}
+
+		const bool bHasDirectLinkSourceChanged = AutoReimportObjectsMap.FindChecked(InAsset)->ExternalSource != UpdatedExternalSource;
+		if (bHasDirectLinkSourceChanged)
+		{
+			// The source changed but is still a DirectLink source.
+			// Since the auto-reimport is asset-driven and not source-driven, keep the auto reimport active with the new source.	
+			DisableAssetAutoReimport(InAsset);
+			EnableAssetAutoReimport(InAsset);
+		}
+	}
+
+	void FAutoReimportManager::OnExternalSourceChanged(const TSharedRef<FExternalSource>& ExternalSource)
+	{
+		// Put the reimport request in a thread-safe queue that will be processed in the main thread.
+		PendingReimportQueue.Enqueue(ExternalSource);
+	}
+
+	void FAutoReimportManager::TriggerAutoReimportOnExternalSource(const TSharedRef<FExternalSource>& ExternalSource)
+	{
+		TArray<TSharedRef<FAutoReimportInfo>> AutoReimportInfos;
+		AutoReimportExternalSourcesMap.MultiFind(ExternalSource, AutoReimportInfos);
+		if (AutoReimportInfos.Num() == 0)
+		{
+			return;
+		}
+
+		for (const TSharedRef<FAutoReimportInfo>& AutoReimportInfo : AutoReimportInfos)
+		{
+			if (UObject* Asset = AutoReimportInfo->TargetObject.Get())
+			{
+				TriggerAutoReimportOnAsset(AutoReimportInfo->TargetObject.Get());
+			}
+		}
+	}
+
+	void FAutoReimportManager::TriggerAutoReimportOnAsset(UObject* Asset)
+	{
+		const FAssetData AssetData(Asset);
+		const FSourceUri Uri = FSourceUri::FromAssetData(AssetData);
+
+		// Make sure we are not triggering a reimport on an asset that doesn't have a DirectLink source.
+		if (Uri.IsValid())
+		{
+			FReimportManager::Instance()->Reimport(Asset, /*bAskForNewFileIfMissing*/ false, /*bShowNotification*/ true, /*PreferredReimportFile*/ TEXT(""), /*SpecifiedReimportHandler */ nullptr, /*SourceFileIndex*/ INDEX_NONE, /*bForceNewFile*/ false, /*bAutomated*/ true);
+		}
+		else
+		{
+			DisableAssetAutoReimport(Asset);
+		}
+	}
+}
\ No newline at end of file
diff --git a/Engine/Plugins/Enterprise/DatasmithImporter/Source/DirectLinkExtension/Private/AutoReimportManager.h b/Engine/Plugins/Enterprise/DatasmithImporter/Source/DirectLinkExtension/Private/AutoReimportManager.h
new file mode 100644
--- /dev/null
+++ b/Engine/Plugins/Enterprise/DatasmithImporter/Source/DirectLinkExtension/Private/AutoReimportManager.h
@@ -0,0 +1,70 @@
+// Copyright Epic Games, Inc. All Rights Reserved.
+
+#pragma once
+
+#include "Containers/Queue.h"
+#include "Containers/Map.h"
+#include "Delegates/IDelegateInstance.h"
+#include "Logging/LogMacros.h"
+#include "Templates/SharedPointer.h"
+#include "TickableEditorObject.h"
+#include "UObject/SoftObjectPtr.h"
+
+class UObject;
+
+DECLARE_LOG_CATEGORY_EXTERN(LogAutoReimportManager, Log, All);
+
+namespace UE::DatasmithImporter
+{
+	class FExternalSource;
+	struct FAutoReimportInfo;
+
+	class FAutoReimportManager : public FTickableEditorObject, public TSharedFromThis<FAutoReimportManager>
+	{
+
+	public:
+		/* Begin FTickableEditorObject Interface */
+		virtual void Tick(float DeltaTime) override;
+		virtual bool IsAllowedToTick() const override { return PendingAutoReimportObjects.Num() > 0 || !PendingReimportQueue.IsEmpty(); }
+		virtual TStatId GetStatId() const override;
+		/* End FTickableEditorObject Interface */
+
+		bool IsAssetAutoReimportEnabled(UObject* InAsset) const;
+
+		bool SetAssetAutoReimport(UObject* InAsset, bool bEnabled);
+
+		void OnExternalSourceInvalidated(const TSharedRef<FExternalSource>& ExternalSource);
+
+		/**
+		 * Update the internal registration a given asset registered for auto-reimport.
+		 * Modified assets may no longer have a DirectLink source and we must keep track of such changes.
+		 */
+		void UpdateModifiedRegisteredAsset(UObject* InAsset);
+
+	private:
+
+		bool CanTriggerReimport() const;
+
+		bool EnableAssetAutoReimport(UObject* InAsset);
+
+		bool DisableAssetAutoReimport(UObject* InAsset);
+
+		void OnExternalSourceChanged(const TSharedRef<FExternalSource>& ExternalSource);
+
+		void TriggerAutoReimportOnExternalSource(const TSharedRef<FExternalSource>& ExternalSource);
+
+		void TriggerAutoReimportOnAsset(UObject* Asset);
+
+	private:
+
+		TSet<TSoftObjectPtr<UObject>> PendingAutoReimportObjects;
+
+		TMap<UObject*, TSharedRef<FAutoReimportInfo>> AutoReimportObjectsMap;
+
+		TMultiMap<TSharedRef<FExternalSource>, TSharedRef<FAutoReimportInfo>> AutoReimportExternalSourcesMap;
+
+		TQueue<TSharedPtr<FExternalSource>> PendingReimportQueue;
+
+		TQueue<TSharedPtr<FExternalSource>> PendingInvalidateQueue;
+	};
+}
\ No newline at end of file
diff --git a/Engine/Plugins/Enterprise/DatasmithImporter/Source/DirectLinkExtension/Private/DirectLinkManager.cpp b/Engine/Plugins/Enterprise/DatasmithImporter/Source/DirectLinkExtension/Private/DirectLinkManager.cpp
--- a/Engine/Plugins/Enterprise/DatasmithImporter/Source/DirectLinkExtension/Private/DirectLinkManager.cpp
+++ b/Engine/Plugins/Enterprise/DatasmithImporter/Source/DirectLinkExtension/Private/DirectLinkManager.cpp
@@ -89,26 +89,6 @@ namespace UE::DatasmithImporter
 			bHasSourceToReconnect = Manager.ExternalSourcesToReconnect.Num() > 0;
 		}
 
-		TArray<TSoftObjectPtr<UObject>> AutoReimportObjectListCopy;
-		{
-			FRWScopeLock ScopeLock(Manager.AutoReimportLock, FRWScopeLockType::SLT_ReadOnly);
-			AutoReimportObjectListCopy = Manager.PendingAutoReimportObjects.Array();
-		}
-
-		for (const TSoftObjectPtr<UObject>& Object : AutoReimportObjectListCopy)
-		{
-			if (Object.IsValid())
-			{
-				Manager.EnableAssetAutoReimport(Object.Get());
-			}
-		}
-
-		if (AutoReimportObjectListCopy.Num() > 0 && !bHasSourceToReconnect)
-		{
-			FRWScopeLock ScopeLock(Manager.AutoReimportLock, FRWScopeLockType::SLT_ReadOnly);
-			bHasSourceToReconnect = Manager.PendingAutoReimportObjects.Num() > 0;
-		}
-
 		LastTryTime = FPlatformTime::Seconds();
 
 		if (bShouldRun && bHasSourceToReconnect)
@@ -118,34 +98,13 @@ namespace UE::DatasmithImporter
 		}
 	}
 
-	/**
-	 * #ueent_todo: The AutoReimport feature should be generalize to all FExternalSource, not just DirectLink ones.
-	 */
-	struct FAutoReimportInfo
-	{
-		FAutoReimportInfo(UObject* InTargetObject, const TSharedRef<FExternalSource>& InExternalSource, FDelegateHandle InImportDelegateHandle)
-			: TargetObject(InTargetObject)
-			, ExternalSource(InExternalSource)
-			, ImportDelegateHandle(InImportDelegateHandle)
-			, bChangedDuringPIE(false)
-		{}
-
-		TSoftObjectPtr<UObject> TargetObject;
-		TSharedRef<FExternalSource> ExternalSource;
-		FDelegateHandle ImportDelegateHandle;
-		bool bChangedDuringPIE;
-	};
-
 	FDirectLinkManager::FDirectLinkManager()
 		: Endpoint(MakeUnique<DirectLink::FEndpoint>(TEXT("UE5-Editor")))
 		, AssetObserver(MakeUnique<FDirectLinkAssetObserver>(*this))
 		, ReconnectionManager(MakeUnique<FDirectLinkAutoReconnectManager>(*this))
+		, AutoReimportManger(MakeShared<FAutoReimportManager>())
 	{
 		Endpoint->AddEndpointObserver(this);
-
-#if WITH_EDITOR
-		OnPIEEndHandle = FEditorDelegates::EndPIE.AddRaw(this, &FDirectLinkManager::OnEndPIE);
-#endif //WITH_EDITOR
 	}
 
 	FDirectLinkManager::~FDirectLinkManager()
@@ -157,10 +116,6 @@ namespace UE::DatasmithImporter
 		{
 			UriExternalSourcePair.Value->Invalidate();
 		}
-
-#if WITH_EDITOR
-		FEditorDelegates::EndPIE.Remove(OnPIEEndHandle);
-#endif //WITH_EDITOR
 	}
 
 	TSharedPtr<FDirectLinkExternalSource> FDirectLinkManager::GetOrCreateExternalSource(const DirectLink::FSourceHandle& SourceHandle)
@@ -410,136 +365,11 @@ namespace UE::DatasmithImporter
 	{
 		TSharedRef<FDirectLinkExternalSource> DirectLinkExternalSource = DirectLinkSourceToExternalSourceMap.FindAndRemoveChecked(InvalidSourceHandle);
 
-		// Clear the auto-reimport cache for this external source.
-		{
-			FRWScopeLock ScopeLock(AutoReimportLock, FRWScopeLockType::SLT_Write);
-			TArray<TSharedRef<FAutoReimportInfo>> AutoReimportInfoList;
-			AutoReimportExternalSourcesMap.MultiFind(DirectLinkExternalSource, AutoReimportInfoList);
-			AutoReimportExternalSourcesMap.Remove(DirectLinkExternalSource);
-
-			for (const TSharedRef<FAutoReimportInfo>& AutoReimportInfo : AutoReimportInfoList)
-			{
-				UObject* Asset = AutoReimportInfo->TargetObject.Get();
-
-				// Assets that were registered for auto-reimport are unregistered from the source
-				// and put back in the "PendingRegistration" list. That way if a compatible source appear they will register to it.
-				AutoReimportObjectsMap.Remove(Asset);
-				PendingAutoReimportObjects.Add(Asset);
-			}
-
-			if (PendingAutoReimportObjects.Num() > 0)
-			{
-				ReconnectionManager->Start();
-			}
-		}
+		AutoReimportManger->OnExternalSourceInvalidated(DirectLinkExternalSource);
 
 		DirectLinkExternalSource->Invalidate();
 	}
 
-	bool FDirectLinkManager::IsAssetAutoReimportEnabled(UObject* InAsset) const
-	{
-		FRWScopeLock ScopeLock(AutoReimportLock, FRWScopeLockType::SLT_ReadOnly);
-
-		return AutoReimportObjectsMap.Contains(InAsset)	|| PendingAutoReimportObjects.Contains(InAsset);
-	}
-
-	bool FDirectLinkManager::SetAssetAutoReimport(UObject* InAsset, bool bEnabled)
-	{
-		return bEnabled ? EnableAssetAutoReimport(InAsset) : DisableAssetAutoReimport(InAsset);
-	}
-
-	bool FDirectLinkManager::EnableAssetAutoReimport(UObject* InAsset)
-	{	
-		// Enable auto-reimport for InAsset.
-		FAssetData AssetData(InAsset);
-		const FSourceUri Uri = FSourceUri::FromAssetData(AssetData);
-		const bool bIsValidDirectLinkUri = Uri.IsValid() && Uri.HasScheme(FDirectLinkUriResolver::GetDirectLinkScheme());
-
-		FRWScopeLock ScopeLock(AutoReimportLock, FRWScopeLockType::SLT_Write);
-		if (bIsValidDirectLinkUri && !AutoReimportObjectsMap.Contains(InAsset))
-		{
-			if (TSharedPtr<FDirectLinkExternalSource> ExternalSource = GetOrCreateExternalSource(Uri))
-			{
-				// Register a delegate triggering a reimport task on the external source snapshotupdate event.
-				// That way the asset will be auto-reimported and kept up-to-date.
-				FDelegateHandle DelegateHandle = ExternalSource->OnExternalSourceChanged.AddRaw(this, &FDirectLinkManager::OnExternalSourceChanged);
-
-				TSharedRef<FDirectLinkExternalSource> ExternalSourceRef = ExternalSource.ToSharedRef();
-				TSharedRef<FAutoReimportInfo> AutoReimportInfo = MakeShared<FAutoReimportInfo>(InAsset, ExternalSourceRef, DelegateHandle);
-
-				AutoReimportObjectsMap.Add(InAsset, AutoReimportInfo);
-				AutoReimportExternalSourcesMap.Add(ExternalSourceRef, AutoReimportInfo);
-				ExternalSource->OpenStream();
-				PendingAutoReimportObjects.Remove(InAsset);
-				return true;
-			}
-			else
-			{
-				bool bIsAlreadySet = false;
-				PendingAutoReimportObjects.Add(InAsset, &bIsAlreadySet);
-				if (!bIsAlreadySet)
-				{
-					ReconnectionManager->Start();
-					return true;
-				}
-			}
-		}
-
-		return false;
-	}
-
-	bool FDirectLinkManager::DisableAssetAutoReimport(UObject* InAsset)
-	{
-		// Disable auto-reimport for InAsset.
-		FRWScopeLock ScopeLock(AutoReimportLock, FRWScopeLockType::SLT_Write);
-		if (const TSharedRef<FAutoReimportInfo>* AutoReimportInfoPtr = AutoReimportObjectsMap.Find(InAsset))
-		{
-			// Holding a local reference to the FAutoReimportInfo to ensure its lifetime while we are cleaning up.
-			const TSharedRef<FAutoReimportInfo> AutoReimportInfo(*AutoReimportInfoPtr);
-
-			AutoReimportInfo->ExternalSource->OnExternalSourceChanged.Remove(AutoReimportInfo->ImportDelegateHandle);
-			AutoReimportObjectsMap.Remove(InAsset);
-			AutoReimportExternalSourcesMap.RemoveSingle(AutoReimportInfo->ExternalSource, AutoReimportInfo);
-			return true;
-		}
-		else
-		{
-			PendingAutoReimportObjects.Remove(InAsset);
-			return true;
-		}
-
-		return false;
-	}
-
-	void FDirectLinkManager::UpdateModifiedRegisteredAsset(UObject* InAsset)
-	{
-		if (!IsAssetAutoReimportEnabled(InAsset))
-		{
-			// Asset is not registered, nothing to update.
-			return;
-		}
-
-		const FAssetData AssetData(InAsset);
-		const FSourceUri Uri = FSourceUri::FromAssetData(AssetData);
-		const bool bIsDirectLinkUri = Uri.IsValid() && Uri.HasScheme(FDirectLinkUriResolver::GetDirectLinkScheme());
-		const TSharedPtr<FExternalSource> UpdatedExternalSource = bIsDirectLinkUri ? GetOrCreateExternalSource(Uri) : nullptr;
-		if (!UpdatedExternalSource)
-		{
-			// Asset was registered for auto reimport but no longer has a DirectLink source, disable auto reimport.
-			DisableAssetAutoReimport(InAsset);
-			return;
-		}
-
-		const bool bHasDirectLinkSourceChanged = AutoReimportObjectsMap.FindChecked(InAsset)->ExternalSource != UpdatedExternalSource;
-		if (bHasDirectLinkSourceChanged)
-		{
-			// The source changed but is still a DirectLink source.
-			// Since the auto-reimport is asset-driven and not source-driven, keep the auto reimport active with the new source.	
-			DisableAssetAutoReimport(InAsset);
-			EnableAssetAutoReimport(InAsset);
-		}
-	}
-
 	TArray<TSharedRef<FDirectLinkExternalSource>> FDirectLinkManager::GetExternalSourceList() const
 	{
 		TArray<TSharedRef<FDirectLinkExternalSource>> ExternalSources;
@@ -566,116 +396,6 @@ namespace UE::DatasmithImporter
 
 		return FSourceUri();
 	}
-
-	void FDirectLinkManager::OnExternalSourceChanged(const TSharedRef<FExternalSource>& ExternalSource)
-	{
-		// Accumulate the reimport request in a thread-safe queue that will be processed in the main thread.
-		// Multiple reimport request for the same external source will only be processed once. 
-		// Doing this allows us to skip redundant reimports, as the reimport already uses the latest data from the ExternalSource.
-		PendingReimportQueue.Enqueue(ExternalSource);
-
-		Async(EAsyncExecution::TaskGraphMainThread, [this]() {
-
-			if (UE::IsSavingPackage(nullptr) || IsGarbageCollecting())
-			{
-				// Can't trigger a reimport while we are saving or garbage collecting.
-				return;
-			}
-
-			TSet<TSharedRef<FExternalSource>> PendingReimportSet;
-			TSharedPtr<FExternalSource> EnqueuedSourceToReimport;
-			while (PendingReimportQueue.Dequeue(EnqueuedSourceToReimport))
-			{
-				PendingReimportSet.Add(EnqueuedSourceToReimport.ToSharedRef());
-			}
-
-			for (const TSharedRef<FExternalSource>& ExternalSourceToReimport : PendingReimportSet)
-			{
-				TriggerAutoReimportOnExternalSource(ExternalSourceToReimport);
-			}
-		});
-	}
-
-	void FDirectLinkManager::TriggerAutoReimportOnExternalSource(const TSharedRef<FExternalSource>& ExternalSource)
-	{
-		TArray<TSharedRef<FAutoReimportInfo>> AutoReimportInfos;
-		AutoReimportExternalSourcesMap.MultiFind(ExternalSource, AutoReimportInfos);
-		if (AutoReimportInfos.Num() == 0)
-		{
-			return;
-		}
-
-		for (const TSharedRef<FAutoReimportInfo>& AutoReimportInfo : AutoReimportInfos)
-		{
-#if WITH_EDITOR
-			// If we're in PIE, delay the callbacks until we exit that mode.
-			if (GIsEditor && FApp::IsGame())
-			{
-				AutoReimportInfo->bChangedDuringPIE = true;
-				UE_LOG(LogDirectLinkManager, Warning, TEXT("The DirectLink source \"%s\" received an update while in PIE mode. The reimport will be triggered when exiting PIE."), *ExternalSource->GetSourceName());
-				continue;
-			}
-#endif //WITH_EDITOR
-
-			if (UObject* Asset = AutoReimportInfo->TargetObject.Get())
-			{
-				TriggerAutoReimportOnAsset(AutoReimportInfo->TargetObject.Get());
-			}
-		}
-	}
-
-	void FDirectLinkManager::TriggerAutoReimportOnAsset(UObject* Asset)
-	{
-		const FAssetData AssetData(Asset);
-		const FSourceUri Uri = FSourceUri::FromAssetData(AssetData);
-		const bool bIsStillValidDirectLinkUri = Uri.IsValid() && Uri.HasScheme(FDirectLinkUriResolver::GetDirectLinkScheme());
-
-		// Make sure we are not triggering a reimport on an asset that doesn't have a DirectLink source.
-		if (bIsStillValidDirectLinkUri)
-		{
-			FReimportManager::Instance()->Reimport(Asset, /*bAskForNewFileIfMissing*/ false, /*bShowNotification*/ true, /*PreferredReimportFile*/ TEXT(""), /*SpecifiedReimportHandler */ nullptr, /*SourceFileIndex*/ INDEX_NONE, /*bForceNewFile*/ false, /*bAutomated*/ true);
-		}
-		else
-		{
-			DisableAssetAutoReimport(Asset);
-		}
-	}
-
-#if WITH_EDITOR
-	void FDirectLinkManager::OnEndPIE(bool bIsSimulating)
-	{
-		TArray<UObject*> AssetsToReimport;
-		TArray<UObject*> InvalidAssets;
-
-		// We can't call TriggerOnExternalSourceChanged() directly as it may remove items in RegisteredAutoReimportObjectMap while we iterate.
-		for (TPair<UObject*, TSharedRef<FAutoReimportInfo>>& AutoReimportEntry : AutoReimportObjectsMap)
-		{
-			if (AutoReimportEntry.Value->TargetObject.IsValid())
-			{
-				if (AutoReimportEntry.Value->bChangedDuringPIE)
-				{
-					AutoReimportEntry.Value->bChangedDuringPIE = false;
-					AssetsToReimport.Add(AutoReimportEntry.Key);
-				}
-			}
-			else
-			{
-				InvalidAssets.Add(AutoReimportEntry.Key);
-			}
-		}
-
-		for (UObject* CurrentAsset : AssetsToReimport)
-		{
-			// Trigger the event held off during PIE.
-			TriggerAutoReimportOnAsset(CurrentAsset);
-		}
-
-		for (UObject* CurrentAsset : InvalidAssets)
-		{
-			DisableAssetAutoReimport(CurrentAsset);
-		}
-	}
-#endif //WITH_EDITOR
 }
 
 #undef LOCTEXT_NAMESPACE
diff --git a/Engine/Plugins/Enterprise/DatasmithImporter/Source/DirectLinkExtension/Private/DirectLinkManager.h b/Engine/Plugins/Enterprise/DatasmithImporter/Source/DirectLinkExtension/Private/DirectLinkManager.h
--- a/Engine/Plugins/Enterprise/DatasmithImporter/Source/DirectLinkExtension/Private/DirectLinkManager.h
+++ b/Engine/Plugins/Enterprise/DatasmithImporter/Source/DirectLinkExtension/Private/DirectLinkManager.h
@@ -4,6 +4,7 @@
 
 #include "IDirectLinkManager.h"
 
+#include "AutoReimportManager.h"
 #include "SourceUri.h"
 #include "Containers/Map.h"
 #include "Containers/UnrealString.h"
@@ -23,7 +24,6 @@ namespace UE::DatasmithImporter
 	class FDirectLinkAssetObserver;
 	class FDirectLinkExternalSource;
 	class FDirectLinkManager;
-	struct FAutoReimportInfo;
 	struct FDirectLinkSourceDescription;
 
 	class FDirectLinkAutoReconnectManager
@@ -69,8 +69,8 @@ namespace UE::DatasmithImporter
 		virtual DirectLink::FEndpoint& GetEndpoint() override;
 		virtual FSourceUri GetUriFromSourceHandle(const DirectLink::FSourceHandle& SourceHandle) override;
 #if WITH_EDITOR
-		virtual bool IsAssetAutoReimportEnabled(UObject* InAsset) const override;
-		virtual bool SetAssetAutoReimport(UObject* InAsset, bool bEnabled) override;
+		virtual bool IsAssetAutoReimportEnabled(UObject* InAsset) const override { return AutoReimportManger->IsAssetAutoReimportEnabled(InAsset); }
+		virtual bool SetAssetAutoReimport(UObject* InAsset, bool bEnabled) override { return AutoReimportManger->SetAssetAutoReimport(InAsset, bEnabled); }
 #endif //WITH_EDITOR
 		virtual TArray<TSharedRef<FDirectLinkExternalSource>> GetExternalSourceList() const override;
 		virtual void UnregisterDirectLinkExternalSource(FName InName) override;
@@ -84,7 +84,7 @@ namespace UE::DatasmithImporter
 		 * Update the internal registration a given asset registered for auto-reimport.
 		 * Modified assets may no longer have a DirectLink source and we must keep track of such changes.
 		 */
-		void UpdateModifiedRegisteredAsset(UObject* InAsset);
+		void UpdateModifiedRegisteredAsset(UObject* InAsset) { AutoReimportManger->UpdateModifiedRegisteredAsset(InAsset); }
 
 	private:
 		/**
@@ -107,20 +107,6 @@ namespace UE::DatasmithImporter
 
 		void CancelEmptySourcesLoading() const;
 
-		bool EnableAssetAutoReimport(UObject* InAsset);
-
-		bool DisableAssetAutoReimport(UObject* InAsset);
-
-		void OnExternalSourceChanged(const TSharedRef<FExternalSource>& ExternalSource);
-
-		void TriggerAutoReimportOnExternalSource(const TSharedRef<FExternalSource>& ExternalSource);
-
-		void TriggerAutoReimportOnAsset(UObject* Asset);
-
-#if WITH_EDITOR
-		void OnEndPIE(bool bIsSimulating);
-#endif
-
 	private:
 		/**
 		 * Cached DirectLink state.
@@ -145,19 +131,8 @@ namespace UE::DatasmithImporter
 
 		TUniquePtr<FDirectLinkAutoReconnectManager> ReconnectionManager;
 		
-		mutable FRWLock AutoReimportLock;
-
-		TSet<TSoftObjectPtr<UObject>> PendingAutoReimportObjects;
-
-		TMap<UObject*, TSharedRef<FAutoReimportInfo>> AutoReimportObjectsMap;
+		TSharedRef<FAutoReimportManager> AutoReimportManger;
 		
-		TMultiMap<TSharedRef<FExternalSource>, TSharedRef<FAutoReimportInfo>> AutoReimportExternalSourcesMap;
-
-		TQueue<TSharedPtr<FExternalSource>> PendingReimportQueue;
-
-#if WITH_EDITOR
-		FDelegateHandle OnPIEEndHandle;
-#endif
 		friend FDirectLinkAutoReconnectManager;
 	};
 }
\ No newline at end of file
