Fix EngineTest cooked editor build - Removing duplicate definition of LogAutoReimportManager.

diff --git a/Engine/Plugins/Enterprise/DatasmithImporter/Source/DirectLinkExtension/Private/AutoReimportManager.cpp b/Engine/Plugins/Enterprise/DatasmithImporter/Source/DirectLinkExtension/Private/AutoReimportManager.cpp
--- a/Engine/Plugins/Enterprise/DatasmithImporter/Source/DirectLinkExtension/Private/AutoReimportManager.cpp
+++ b/Engine/Plugins/Enterprise/DatasmithImporter/Source/DirectLinkExtension/Private/AutoReimportManager.cpp
@@ -11,8 +11,6 @@
 #include "Stats/Stats.h"
 #include "Misc/App.h"
 
-DEFINE_LOG_CATEGORY(LogAutoReimportManager);
-
 namespace UE::DatasmithImporter
 {
 	struct FAutoReimportInfo
diff --git a/Engine/Plugins/Enterprise/DatasmithImporter/Source/DirectLinkExtension/Private/AutoReimportManager.h b/Engine/Plugins/Enterprise/DatasmithImporter/Source/DirectLinkExtension/Private/AutoReimportManager.h
--- a/Engine/Plugins/Enterprise/DatasmithImporter/Source/DirectLinkExtension/Private/AutoReimportManager.h
+++ b/Engine/Plugins/Enterprise/DatasmithImporter/Source/DirectLinkExtension/Private/AutoReimportManager.h
@@ -12,8 +12,6 @@
 
 class UObject;
 
-DECLARE_LOG_CATEGORY_EXTERN(LogAutoReimportManager, Log, All);
-
 namespace UE::DatasmithImporter
 {
 	class FExternalSource;
