Dataprep - unable to add components to actors upon commit

diff --git a/Engine/Plugins/Enterprise/DatasmithImporter/Source/DatasmithImporter/Private/DatasmithImporter.cpp b/Engine/Plugins/Enterprise/DatasmithImporter/Source/DatasmithImporter/Private/DatasmithImporter.cpp
--- a/Engine/Plugins/Enterprise/DatasmithImporter/Source/DatasmithImporter/Private/DatasmithImporter.cpp
+++ b/Engine/Plugins/Enterprise/DatasmithImporter/Source/DatasmithImporter/Private/DatasmithImporter.cpp
@@ -1346,6 +1346,9 @@ AActor* FDatasmithImporter::FinalizeActor( FDatasmithImportContext& ImportContex
 		// The templates for the actor need to be applied after the components were created.
 		FDatasmithImporterImpl::ApplyMigratedTemplates( MigratedTemplates, DestinationActor );
 
+		// Actors that were created via blueprint (can happen with Dataprep) will have their components destroyed by the finalizer after this scope's end.
+		// This will disallow destruction of such components.
+		DestinationActor->bActorSeamlessTraveled = 1;
 	}
 
 	// Update label to match the source actor's
