DirectLink in Editor - Fixing a crash that could happen when closing the editor.

There was a race condition happening during the DirectLinkManager destructor, that might cause the direct link Endpoint to access data that has already been destroyed if the DirectLink background thread happens to be processing updates at the same time.
The FDirectLinkExternalSource invalidation was relying on the DirectLinkExtension module to be available, but since their destruction was happening after the call to ShutdownModule() this was not the case.
And so, they could not properly unregister themselves from the Endpoint and  for a brief moment the Endpoint could have dangling pointers.
The crash would happen if the DirectLink background thread happened to run while the Endpoint was in that state.

To fix this, we now make sure that the DirectLinkManager is  properly destroyed during ShutdownModule() and that FDirectLinkExternalSources now use a static DirectLinkManager instance rather than relying on the DirectLinkExtension module to unregister themselves.

diff --git a/Engine/Plugins/Enterprise/DatasmithImporter/Source/DirectLinkExtension/Private/DirectLinkExtensionModule.cpp b/Engine/Plugins/Enterprise/DatasmithImporter/Source/DirectLinkExtension/Private/DirectLinkExtensionModule.cpp
--- a/Engine/Plugins/Enterprise/DatasmithImporter/Source/DirectLinkExtension/Private/DirectLinkExtensionModule.cpp
+++ b/Engine/Plugins/Enterprise/DatasmithImporter/Source/DirectLinkExtension/Private/DirectLinkExtensionModule.cpp
@@ -26,7 +26,8 @@ namespace UE::DatasmithImporter
 	public:
 		virtual void StartupModule() override
 		{
-			DirectLinkManager = MakeUnique<FDirectLinkManager>();
+			// This will instantiate the DirectLinkManager.
+			FDirectLinkManager::GetInstance();
 			DirectLinkExtensionUI = MakeUnique<FDirectLinkExtensionUI>();
 
 			IUriManager& UriManager = IExternalSourceModule::Get().GetManager();
@@ -40,18 +41,19 @@ namespace UE::DatasmithImporter
 				IUriManager& UriManager = IExternalSourceModule::Get().GetManager();
 				UriManager.UnregisterResolver(DirectLinkUriResolverName);
 			}
+
+			DirectLinkExtensionUI.Reset();
+			FDirectLinkManager::ResetInstance();
 		}
 
 		virtual IDirectLinkManager& GetManager() const override
 		{
-			check(DirectLinkManager.IsValid());
-			return *DirectLinkManager;
+			return FDirectLinkManager::GetInstance();
 		}
 
 		virtual TSharedPtr<FDirectLinkExternalSource> DisplayDirectLinkSourcesDialog() override;
 
 	private:
-		TUniquePtr<FDirectLinkManager> DirectLinkManager;
 		TUniquePtr<FDirectLinkExtensionUI> DirectLinkExtensionUI;
 	};
 
diff --git a/Engine/Plugins/Enterprise/DatasmithImporter/Source/DirectLinkExtension/Private/DirectLinkExternalSource.cpp b/Engine/Plugins/Enterprise/DatasmithImporter/Source/DirectLinkExtension/Private/DirectLinkExternalSource.cpp
--- a/Engine/Plugins/Enterprise/DatasmithImporter/Source/DirectLinkExtension/Private/DirectLinkExternalSource.cpp
+++ b/Engine/Plugins/Enterprise/DatasmithImporter/Source/DirectLinkExtension/Private/DirectLinkExternalSource.cpp
@@ -4,7 +4,7 @@
 #include "DirectLinkExternalSource.h"
 
 #include "DirectLinkExtensionModule.h"
-#include "IDirectLinkManager.h"
+#include "DirectLinkManager.h"
 
 #include "Async/Async.h"
 #include "DirectLinkEndpoint.h"
@@ -71,7 +71,7 @@ namespace UE::DatasmithImporter
 
 	FDirectLinkExternalSource::~FDirectLinkExternalSource()
 	{
-		Invalidate();
+		ensureMsgf(!DestinationHandle.IsValid(), TEXT("Invalidate() was not called before the FDirectLinkExternalSource destruction. This may have left dangling pointers in the Endpoint."));
 	}
 
 	FExternalSourceCapabilities FDirectLinkExternalSource::GetCapabilities() const
@@ -95,7 +95,7 @@ namespace UE::DatasmithImporter
 
 		if (SourceHandle.IsValid() && DestinationHandle.IsValid())
 		{
-			const FEndpoint::EOpenStreamResult Result = IDirectLinkExtensionModule::GetEndpoint().OpenStream(SourceHandle, DestinationHandle);
+			const FEndpoint::EOpenStreamResult Result = FDirectLinkManager::GetInstance().GetEndpoint().OpenStream(SourceHandle, DestinationHandle);
 
 			bIsStreamOpen = Result == FEndpoint::EOpenStreamResult::Opened || Result == FEndpoint::EOpenStreamResult::AlreadyOpened;
 
@@ -109,18 +109,15 @@ namespace UE::DatasmithImporter
 	{
 		if (SourceHandle.IsValid() && DestinationHandle.IsValid())
 		{
-			IDirectLinkExtensionModule::GetEndpoint().CloseStream(SourceHandle, DestinationHandle);
+			FDirectLinkManager::GetInstance().GetEndpoint().CloseStream(SourceHandle, DestinationHandle);
 			bIsStreamOpen = false;
 		}
 	}
 
 	void FDirectLinkExternalSource::Invalidate()
 	{
-		if (IDirectLinkExtensionModule::IsAvailable())
-		{
-			CloseStream();
-			IDirectLinkExtensionModule::GetEndpoint().RemoveDestination(DestinationHandle);
-		}
+		CloseStream();
+		FDirectLinkManager::GetInstance().GetEndpoint().RemoveDestination(DestinationHandle);
 
 		ClearOnExternalSourceLoadedDelegates();
 		DestinationHandle.Invalidate();
diff --git a/Engine/Plugins/Enterprise/DatasmithImporter/Source/DirectLinkExtension/Private/DirectLinkManager.cpp b/Engine/Plugins/Enterprise/DatasmithImporter/Source/DirectLinkExtension/Private/DirectLinkManager.cpp
--- a/Engine/Plugins/Enterprise/DatasmithImporter/Source/DirectLinkExtension/Private/DirectLinkManager.cpp
+++ b/Engine/Plugins/Enterprise/DatasmithImporter/Source/DirectLinkExtension/Private/DirectLinkManager.cpp
@@ -25,6 +25,8 @@ DEFINE_LOG_CATEGORY(LogDirectLinkManager);
 
 namespace UE::DatasmithImporter
 {
+	TUniquePtr<FDirectLinkManager> FDirectLinkManager::Instance = nullptr;
+
 	FDirectLinkAutoReconnectManager::FDirectLinkAutoReconnectManager(FDirectLinkManager& InManager)
 		: Manager(InManager)
 		, bShouldRun(true)
@@ -105,9 +107,11 @@ namespace UE::DatasmithImporter
 		, AutoReimportManger(MakeShared<FAutoReimportManager>())
 	{
 		Endpoint->AddEndpointObserver(this);
+
+		ensureMsgf(!Instance, TEXT("There can only be one instance of FDirectLinkManager."));
 	}
 
-	FDirectLinkManager::~FDirectLinkManager()
+	void FDirectLinkManager::Clear()
 	{
 		Endpoint->RemoveEndpointObserver(this);
 
@@ -129,7 +133,7 @@ namespace UE::DatasmithImporter
 		{
 			DirectLink::FRawInfo::FDataPointInfo* SourceDataPointInfo;
 			{
-				FRWScopeLock ScopeLock(RawInfoLock, FRWScopeLockType::SLT_ReadOnly);
+			FRWScopeLock ScopeLock(RawInfoLock, FRWScopeLockType::SLT_ReadOnly);
 				SourceDataPointInfo = RawInfoCache.DataPointsInfo.Find(SourceHandle);
 			}
 			FSourceUri ExternalSourceUri(GetUriFromSourceHandle(SourceHandle));
diff --git a/Engine/Plugins/Enterprise/DatasmithImporter/Source/DirectLinkExtension/Private/DirectLinkManager.h b/Engine/Plugins/Enterprise/DatasmithImporter/Source/DirectLinkExtension/Private/DirectLinkManager.h
--- a/Engine/Plugins/Enterprise/DatasmithImporter/Source/DirectLinkExtension/Private/DirectLinkManager.h
+++ b/Engine/Plugins/Enterprise/DatasmithImporter/Source/DirectLinkExtension/Private/DirectLinkManager.h
@@ -56,8 +56,27 @@ namespace UE::DatasmithImporter
 	class FDirectLinkManager: public IDirectLinkManager, public DirectLink::IEndpointObserver
 	{
 	public:
-		FDirectLinkManager();
-		virtual ~FDirectLinkManager();
+		/**
+		 * Static instance that can be used internally to avoid relying on IDirectLinkExtensionModule to get the DirectLink manager.
+		 */
+		static FDirectLinkManager& GetInstance()
+		{
+			if (!Instance)
+			{
+				Instance = TUniquePtr<FDirectLinkManager>(new FDirectLinkManager());;
+				check(Instance);
+			}
+			return *Instance;
+		}
+		
+		static void ResetInstance()
+		{
+			if (Instance)
+			{
+				Instance->Clear();
+			}
+			Instance.Reset();
+		}
 
 		// IEndpointObserver interface begin
 		virtual void OnStateChanged(const DirectLink::FRawInfo& RawInfo) override;
@@ -87,6 +106,14 @@ namespace UE::DatasmithImporter
 		void UpdateModifiedRegisteredAsset(UObject* InAsset) { AutoReimportManger->UpdateModifiedRegisteredAsset(InAsset); }
 
 	private:
+
+		FDirectLinkManager();
+
+		/**
+		 * Should be called before destructing the FDirectLinkManager.
+		 */
+		void Clear();
+
 		/**
 		 * Remove a DirectLink source from cache and invalidate its associated DirectLinkExternalSource object.
 		 * @param InvalidSourceId	The SourceHandle of the invalid DirectLink source.
@@ -108,6 +135,9 @@ namespace UE::DatasmithImporter
 		void CancelEmptySourcesLoading() const;
 
 	private:
+
+		static TUniquePtr<FDirectLinkManager> Instance;
+
 		/**
 		 * Cached DirectLink state.
 		 */
