DirectLink in Editor - Fixing Mac compilation error.

Default Destructor needs to be defined in the .cpp so it can use the types that are forward declared in the header.

diff --git a/Engine/Plugins/Enterprise/DatasmithImporter/Source/DirectLinkExtension/Private/DirectLinkManager.cpp b/Engine/Plugins/Enterprise/DatasmithImporter/Source/DirectLinkExtension/Private/DirectLinkManager.cpp
--- a/Engine/Plugins/Enterprise/DatasmithImporter/Source/DirectLinkExtension/Private/DirectLinkManager.cpp
+++ b/Engine/Plugins/Enterprise/DatasmithImporter/Source/DirectLinkExtension/Private/DirectLinkManager.cpp
@@ -111,6 +111,8 @@ namespace UE::DatasmithImporter
 		ensureMsgf(!Instance, TEXT("There can only be one instance of FDirectLinkManager."));
 	}
 
+	FDirectLinkManager::~FDirectLinkManager() = default;
+
 	void FDirectLinkManager::Clear()
 	{
 		Endpoint->RemoveEndpointObserver(this);
diff --git a/Engine/Plugins/Enterprise/DatasmithImporter/Source/DirectLinkExtension/Private/DirectLinkManager.h b/Engine/Plugins/Enterprise/DatasmithImporter/Source/DirectLinkExtension/Private/DirectLinkManager.h
--- a/Engine/Plugins/Enterprise/DatasmithImporter/Source/DirectLinkExtension/Private/DirectLinkManager.h
+++ b/Engine/Plugins/Enterprise/DatasmithImporter/Source/DirectLinkExtension/Private/DirectLinkManager.h
@@ -56,6 +56,9 @@ namespace UE::DatasmithImporter
 	class FDirectLinkManager: public IDirectLinkManager, public DirectLink::IEndpointObserver
 	{
 	public:
+
+		virtual ~FDirectLinkManager();
+
 		/**
 		 * Static instance that can be used internally to avoid relying on IDirectLinkExtensionModule to get the DirectLink manager.
 		 */
