Datasmith cleanup:
- update AssetRegistryModule.h include path
- fix malformed visual assist todo tags
- removed unactioned wip tags

diff --git a/Engine/Plugins/Enterprise/DatasmithFBXImporter/Source/DatasmithDeltaGenTranslator/Private/DatasmithDeltaGenImporter.cpp b/Engine/Plugins/Enterprise/DatasmithFBXImporter/Source/DatasmithDeltaGenTranslator/Private/DatasmithDeltaGenImporter.cpp
--- a/Engine/Plugins/Enterprise/DatasmithFBXImporter/Source/DatasmithDeltaGenTranslator/Private/DatasmithDeltaGenImporter.cpp
+++ b/Engine/Plugins/Enterprise/DatasmithFBXImporter/Source/DatasmithDeltaGenTranslator/Private/DatasmithDeltaGenImporter.cpp
@@ -121,7 +121,6 @@ bool FDatasmithDeltaGenImporter::ParseFbxFile(const FString& FBXPath)
 		return false;
 	}
 
-	// #ueent_wip_baseoptions
 	FDatasmithImportBaseOptions DefaultBaseOptions;
 
 	FDatasmithFBXFileImporter Importer(FbxImporter->Scene, IntermediateScene.Get(), ImportOptions, &DefaultBaseOptions);
diff --git a/Engine/Plugins/Enterprise/DatasmithImporter/Source/ExternalSource/Public/ExternalSource.h b/Engine/Plugins/Enterprise/DatasmithImporter/Source/ExternalSource/Public/ExternalSource.h
--- a/Engine/Plugins/Enterprise/DatasmithImporter/Source/ExternalSource/Public/ExternalSource.h
+++ b/Engine/Plugins/Enterprise/DatasmithImporter/Source/ExternalSource/Public/ExternalSource.h
@@ -96,13 +96,13 @@ namespace UE::DatasmithImporter
 
 		/**
 		 * Return the name of the scene that will be loaded.
-		 * #ueent-todo This is only required because of the way the Datasmith reimport works. Consider removing it when we adapt the ExternalSource to interchange.
+		 * #ueent_todo This is only required because of the way the Datasmith reimport works. Consider removing it when we adapt the ExternalSource to interchange.
 		 */
 		FString GetSceneName() const;
 
 		/**
 		 * Override the name of the scene that will be loaded. If the pointer is nullptr, the override will be reset.
-		 * #ueent-todo This is only required because of the way the Datasmith reimport works. Consider removing it when we adapt the ExternalSource to interchange.
+		 * #ueent_todo This is only required because of the way the Datasmith reimport works. Consider removing it when we adapt the ExternalSource to interchange.
 		 */
 		void SetSceneName(const TCHAR* SceneName);
 
diff --git a/Engine/Source/Programs/Enterprise/Datasmith/DatasmithRhinoExporter/Private/DirectLink/DatasmithRhinoChangeListener.cs b/Engine/Source/Programs/Enterprise/Datasmith/DatasmithRhinoExporter/Private/DirectLink/DatasmithRhinoChangeListener.cs
--- a/Engine/Source/Programs/Enterprise/Datasmith/DatasmithRhinoExporter/Private/DirectLink/DatasmithRhinoChangeListener.cs
+++ b/Engine/Source/Programs/Enterprise/Datasmith/DatasmithRhinoExporter/Private/DirectLink/DatasmithRhinoChangeListener.cs
@@ -62,7 +62,7 @@ namespace DatasmithRhino.DirectLink
 					RhinoDoc.TextureMappingEvent += OnTextureMappingEvent;
 					RhinoDoc.DocumentPropertiesChanged += OnDocumentPropertiesChanged;
 
-					//#ueent-todo Listen to the following events to update their associated DatasmithElement
+					//#ueent_todo Listen to the following events to update their associated DatasmithElement
 					//RhinoDoc.DimensionStyleTableEvent;
 					//RhinoDoc.LightTableEvent;
 					//RhinoDoc.RenderEnvironmentTableEvent;
diff --git a/Engine/Source/Programs/Enterprise/Datasmith/DatasmithRhinoExporter/Private/ExportContext/DatasmithRhinoExportContext.cs b/Engine/Source/Programs/Enterprise/Datasmith/DatasmithRhinoExporter/Private/ExportContext/DatasmithRhinoExportContext.cs
--- a/Engine/Source/Programs/Enterprise/Datasmith/DatasmithRhinoExporter/Private/ExportContext/DatasmithRhinoExportContext.cs
+++ b/Engine/Source/Programs/Enterprise/Datasmith/DatasmithRhinoExporter/Private/ExportContext/DatasmithRhinoExportContext.cs
@@ -395,7 +395,7 @@ namespace DatasmithRhino.ExportContext
 			bool bHasMesh = false;
 			if (ObjectIdToMeshInfoDictionary.TryGetValue(InRhinoObject.Id, out DatasmithMeshInfo MeshInfo))
 			{
-				// #ueent-todo	The MeshInfo is the one holding the OffsetTransform but changing it actually affects the actor and not the mesh. 
+				// #ueent_todo	The MeshInfo is the one holding the OffsetTransform but changing it actually affects the actor and not the mesh.
 				//				It is weird that the actor is the one modified here, we should move the OffsetTransform into it's associated DatasmithActorInfo.
 				MeshInfo.OffsetTransform = InTransform * MeshInfo.OffsetTransform;
 				bHasMesh = true;
