DatasmithImporter - Fix crash caused by threadsafety issue when logging errors.

When importing a scene using a lot of unsupported texture formats, the asynchronous logging of those errors would make the editor crash.

diff --git a/Engine/Plugins/Enterprise/DatasmithImporter/Source/DatasmithImporter/Private/Utility/DatasmithImporterUtils.cpp b/Engine/Plugins/Enterprise/DatasmithImporter/Source/DatasmithImporter/Private/Utility/DatasmithImporterUtils.cpp
--- a/Engine/Plugins/Enterprise/DatasmithImporter/Source/DatasmithImporter/Private/Utility/DatasmithImporterUtils.cpp
+++ b/Engine/Plugins/Enterprise/DatasmithImporter/Source/DatasmithImporter/Private/Utility/DatasmithImporterUtils.cpp
@@ -1376,6 +1376,8 @@ FScopedLogger::~FScopedLogger()
 
 TSharedRef<FTokenizedMessage> FScopedLogger::Push(EMessageSeverity::Type Severity, const FText& Message)
 {
+	FScopeLock Lock(&TokenizedMessageCS);
+
 	TokenizedMessages.Add(FTokenizedMessage::Create(Severity, Message));
 
 	return TokenizedMessages.Last();
@@ -1383,6 +1385,8 @@ TSharedRef<FTokenizedMessage> FScopedLogger::Push(EMessageSeverity::Type Severit
 
 void FScopedLogger::Dump(bool bClearPrevious)
 {
+	FScopeLock Lock(&TokenizedMessageCS);
+
 	if (TokenizedMessages.Num() > 0)
 	{
 		if (bClearPrevious)
@@ -1403,6 +1407,8 @@ void FScopedLogger::ClearLog()
 
 void FScopedLogger::ClearPending()
 {
+	FScopeLock Lock(&TokenizedMessageCS);
+
 	TokenizedMessages.Empty();
 }
 
diff --git a/Engine/Plugins/Enterprise/DatasmithImporter/Source/DatasmithImporter/Public/DatasmithImportContext.h b/Engine/Plugins/Enterprise/DatasmithImporter/Source/DatasmithImporter/Public/DatasmithImportContext.h
--- a/Engine/Plugins/Enterprise/DatasmithImporter/Source/DatasmithImporter/Public/DatasmithImportContext.h
+++ b/Engine/Plugins/Enterprise/DatasmithImporter/Source/DatasmithImporter/Public/DatasmithImportContext.h
@@ -93,6 +93,7 @@ private:
 	FMessageLogModule& MessageLogModule;
 	TSharedPtr<IMessageLogListing> LogListing;
 	TArray<TSharedRef<FTokenizedMessage>> TokenizedMessages;
+	FCriticalSection TokenizedMessageCS;
 };
 
 struct FDatasmithActorImportContext
