Datasmith: Cloth support prototype

diff --git a/Engine/Plugins/Enterprise/DatasmithImporter/Source/DatasmithImporter/Private/DatasmithImportContext.cpp b/Engine/Plugins/Enterprise/DatasmithImporter/Source/DatasmithImporter/Private/DatasmithImportContext.cpp
--- a/Engine/Plugins/Enterprise/DatasmithImporter/Source/DatasmithImporter/Private/DatasmithImportContext.cpp
+++ b/Engine/Plugins/Enterprise/DatasmithImporter/Source/DatasmithImporter/Private/DatasmithImportContext.cpp
@@ -230,7 +230,7 @@ FDatasmithImportContext::FDatasmithImportContext(const TSharedPtr<UE::DatasmithI
 		{
 			SceneTranslator->GetSceneImportOptions(AdditionalImportOptions);
 		}
-		
+
 		const FString FileName = InExternalSource->GetFallbackFilepath();
 		Options->FileName = FPaths::GetCleanFilename(FileName);
 		Options->FilePath = FPaths::ConvertRelativePathToFull(FileName);
@@ -555,6 +555,11 @@ void FDatasmithImportContext::FInternalReferenceCollector::AddReferencedObjects(
 		Collector.AddReferencedObject( It->Value );
 	}
 
+	for (auto& Pair : ImportContext->ImportedClothes)
+	{
+		Collector.AddReferencedObject(Pair.Value);
+	}
+
 	Collector.AddReferencedObjects(ImportContext->ImportedMaterials);
 	Collector.AddReferencedObjects(ImportContext->ImportedParentMaterials);
 
diff --git a/Engine/Plugins/Enterprise/DatasmithImporter/Source/DatasmithImporter/Private/DatasmithImportFactory.cpp b/Engine/Plugins/Enterprise/DatasmithImporter/Source/DatasmithImporter/Private/DatasmithImportFactory.cpp
--- a/Engine/Plugins/Enterprise/DatasmithImporter/Source/DatasmithImporter/Private/DatasmithImportFactory.cpp
+++ b/Engine/Plugins/Enterprise/DatasmithImporter/Source/DatasmithImporter/Private/DatasmithImportFactory.cpp
@@ -262,6 +262,15 @@ namespace DatasmithImportFactoryImpl
 			return false;
 		}
 
+		// CLOTH
+		FDatasmithImporter::ImportClothes( InContext );
+
+		if ( InContext.bUserCancelled )
+		{
+			bOutOperationCancelled = true;
+			return false;
+		}
+
 		// ACTORS
 		if( InContext.ShouldImportActors() )
 		{
@@ -496,7 +505,7 @@ UObject* UDatasmithImportFactory::FactoryCreateFile(UClass* InClass, UObject* In
 		UE_LOG(LogDatasmithImport, Warning, TEXT("Datasmith import error: no suitable external source found for this file path. Abort import."));
 		return nullptr;
 	}
-	
+
 	return CreateFromExternalSource(InClass, InParent, InName, InFlags, ExternalSource.ToSharedRef(), InParms, InWarn, bOutOperationCanceled);
 }
 
@@ -843,7 +852,7 @@ EReimportResult::Type UDatasmithImportFactory::ReimportScene(UDatasmithScene* Sc
 	}
 
 	UDatasmithSceneImportData& AssetImportData = *SceneAsset->AssetImportData;
-	
+
 	TSharedPtr<FExternalSource> ExternalSource = IExternalSourceModule::Get().GetManager().TryGetExternalSourceFromImportData(AssetImportData);
 	if (!ExternalSource)
 	{
diff --git a/Engine/Plugins/Enterprise/DatasmithImporter/Source/DatasmithImporter/Private/DatasmithImporter.cpp b/Engine/Plugins/Enterprise/DatasmithImporter/Source/DatasmithImporter/Private/DatasmithImporter.cpp
--- a/Engine/Plugins/Enterprise/DatasmithImporter/Source/DatasmithImporter/Private/DatasmithImporter.cpp
+++ b/Engine/Plugins/Enterprise/DatasmithImporter/Source/DatasmithImporter/Private/DatasmithImporter.cpp
@@ -8,7 +8,10 @@
 #include "DatasmithAssetImportData.h"
 #include "DatasmithAssetUserData.h"
 #include "DatasmithCameraImporter.h"
+#include "DatasmithCloth.h"
+#include "DatasmithClothImporter.h"
 #include "DatasmithImportContext.h"
+#include "DatasmithImporterModule.h"
 #include "DatasmithLevelSequenceImporter.h"
 #include "DatasmithLevelVariantSetsImporter.h"
 #include "DatasmithLightImporter.h"
@@ -374,6 +377,84 @@ void FDatasmithImporter::ImportStaticMeshes( FDatasmithImportContext& ImportCont
 	}
 }
 
+void FDatasmithImporter::ImportClothes(FDatasmithImportContext& ImportContext)
+{
+	// #ue_ds_cloth_note: FDatasmithImporter::ImportClothes: experimental import code
+	const int32 ClothesCount = ImportContext.FilteredScene->GetClothesCount();
+	if ( !ImportContext.Options->BaseOptions.bIncludeGeometry || ClothesCount == 0 )
+	{
+		return;
+	}
+
+	for (int32 ClothIndex = 0; ClothIndex < ClothesCount && !ImportContext.bUserCancelled; ++ClothIndex)
+	{
+		ImportContext.bUserCancelled = FDatasmithImporterImpl::HasUserCancelledTask(ImportContext.FeedbackContext);
+
+		TSharedPtr<IDatasmithClothElement> ClothElementPtr = ImportContext.FilteredScene->GetCloth(ClothIndex);
+		if (!ClothElementPtr)
+		{
+			continue;
+		}
+		TSharedRef<IDatasmithClothElement> ClothElement = ClothElementPtr.ToSharedRef();
+
+		// #ue_ds_cloth_todo: check existing asset
+
+		{ // FDatasmithImporter::ImportStaticMesh
+
+			{ // FDatasmithStaticMeshImporter::ImportStaticMesh
+//
+				// #ue_ds_cloth_todo: import UChaosClothPreset here
+// 				ClothElement->GetPhyDescriptions.....
+
+// 				// 2. get destination asset
+// 				// 2.1. FDatasmithImporterUtils::CanCreateAsset<UStaticMesh>
+				int32 MaxNameCharCount = FDatasmithImporterUtils::GetAssetNameMaxCharCount(ImportContext.AssetsContext.StaticMeshesFinalPackage.Get());
+				FString ClothName = ImportContext.AssetsContext.StaticMeshNameProvider.GenerateUniqueName(ClothElement->GetLabel(), MaxNameCharCount);
+//
+// 				// if !existing
+				UPackage* Outer = ImportContext.AssetsContext.StaticMeshesImportPackage.Get();
+				EObjectFlags ObjectFlags = ImportContext.ObjectFlags & ~RF_Public; // not RF_Public, the publicize asset will be public
+
+				FDatasmithCloth DsCloth;
+
+				if (ImportContext.SceneTranslator)
+				{
+					FDatasmithClothElementPayload Payload;
+					if (ImportContext.SceneTranslator->LoadCloth(ClothElement, Payload))
+					{
+						DsCloth = MoveTemp(Payload.Cloth);
+					}
+					else
+					{
+						UE_LOG(LogDatasmithImport, Warning, TEXT("Cloth %s: load failed, translator issue."), ClothElement->GetName());
+						return;
+					}
+				}
+
+				UObject* ClothAssetObj = nullptr;
+
+#if 0 // #ue_ds_cloth_arch switch between local impl or delegates to IDatasmithImporterExt
+				UChaosClothAsset* ClothAsset = NewObject<UChaosClothAsset>(Outer, *ClothName, ObjectFlags);
+				// #ue_ds_cloth_todo: markPackageDirty?
+				UE::Datasmith::ClothImporter::FillCloth(ClothAsset, ClothElement, DsCloth);
+				ClothAssetObj = (UObject*)ClothAsset;
+#else
+				if (IDatasmithImporterExt* Ext = IDatasmithImporterModule::Get().GetClothImporterExtension())
+				{
+					ClothAssetObj = Ext->MakeClothAsset(Outer, *ClothName, ObjectFlags);
+					Ext->FillCloth(ClothAssetObj, ClothElement, DsCloth);
+				}
+#endif
+
+				ImportContext.ImportedClothes.Add(ClothElement, ClothAssetObj);
+			}
+
+			// #ue_ds_cloth_todo: CreateStaticMeshAssetImportData
+			// #ue_ds_cloth_todo: ImportMetaDataForObject
+		}
+	}
+}
+
 UStaticMesh* FDatasmithImporter::ImportStaticMesh( FDatasmithImportContext& ImportContext, TSharedRef< IDatasmithMeshElement > MeshElement, UStaticMesh* ExistingStaticMesh, FDatasmithMeshElementPayload* MeshPayload)
 {
 	if ( !ImportContext.AssetsContext.StaticMeshesFinalPackage || ImportContext.AssetsContext.StaticMeshesFinalPackage->GetFName() == NAME_None || ImportContext.SceneTranslator == nullptr)
@@ -445,6 +526,14 @@ UStaticMesh* FDatasmithImporter::FinalizeStaticMesh( UStaticMesh* SourceStaticMe
 	return DestinationStaticMesh;
 }
 
+UObject* FDatasmithImporter::FinalizeCloth(UObject* SourceCloth, const TCHAR* FolderPath, UObject* ExistingCloth, TMap<UObject*, UObject*>* ReferencesToRemap)
+{
+	// #ue_ds_cloth_note FDatasmithImporter::FinalizeCloth
+	UObject* Asset = FDatasmithImporterImpl::FinalizeAsset(SourceCloth, FolderPath, ExistingCloth, ReferencesToRemap);
+// 	UChaosClothAsset* DestinationCloth = Cast<UChaosClothAsset>(Asset);
+	return Asset;
+}
+
 void FDatasmithImporter::CreateStaticMeshAssetImportData(FDatasmithImportContext& InContext, TSharedRef< IDatasmithMeshElement > MeshElement, UStaticMesh* ImportedStaticMesh, TArray<UDatasmithAdditionalData*>& AdditionalData)
 {
 	UDatasmithStaticMeshImportData::DefaultOptionsPair ImportOptions = UDatasmithStaticMeshImportData::DefaultOptionsPair( InContext.Options->BaseOptions.StaticMeshOptions, InContext.Options->BaseOptions.AssetOptions );
@@ -1709,6 +1798,7 @@ void FDatasmithImporter::FinalizeImport(FDatasmithImportContext& ImportContext,
 	TRACE_CPUPROFILER_EVENT_SCOPE(FDatasmithImporter::FinalizeImport);
 
 	const int32 NumImportedObjects = ImportContext.ImportedStaticMeshes.Num() +
+									 ImportContext.ImportedClothes.Num() +
 									 ImportContext.ImportedTextures.Num() +
 									 ImportContext.ImportedMaterialFunctions.Num() +
 									 ImportContext.ImportedMaterials.Num() +
@@ -1910,6 +2000,35 @@ void FDatasmithImporter::FinalizeImport(FDatasmithImportContext& ImportContext,
 
 	FDatasmithStaticMeshImporter::BuildStaticMeshes(StaticMeshes.Array(), ProgressFunction);
 
+	// #ue_ds_cloth_note: experimental FinalizeImport code
+	for (auto& ImportedPair : ImportContext.ImportedClothes)
+	{
+		UObject* SourceClothObj = ImportedPair.Value;
+
+		if (!SourceClothObj || (ValidAssets.Num() > 0 && !ValidAssets.Contains(SourceClothObj)))
+		{
+			continue;
+		}
+
+		FDatasmithImporterImpl::ReportProgress(Progress, 1.f, FText::FromString(FString::Printf(TEXT("Finalizing assets %d/%d (Cloth %s) ..."), ++AssetIndex, NumAssetsToFinalize, *SourceClothObj->GetName())));
+
+		FName ClothId = ImportedPair.Key->GetName();
+		UObject* ExistingCloth = nullptr; // #ue_ds_cloth_todo but requires a dependency DatasmithContent -> ChaosClothAsset.uplugin
+// 		UStaticMesh* ExistingCloth = ImportContext.SceneAsset ? ImportContext.SceneAsset->Clothes.FindOrAdd(ClothId).Get() : nullptr;
+
+		FString SourcePackagePath = SourceClothObj->GetOutermost()->GetName();
+		FString DestinationPackagePath = SourcePackagePath.Replace( *TransientFolderPath, *RootFolderPath, ESearchCase::CaseSensitive );
+
+		UObject* FinalizedCloth = FinalizeCloth(SourceClothObj, *DestinationPackagePath, ExistingCloth, &ReferencesToRemap);
+// 		if (ImportContext.SceneAsset)
+// 		{
+// 			ImportContext.SceneAsset->Clothes[ClothId] = FinalizedCloth; // #ue_ds_cloth_todo same constraint
+// 		}
+		FDatasmithImporterImpl::CheckAssetPersistenceValidity(FinalizedCloth, ImportContext);
+
+		ImportedPair.Value = FinalizedCloth;
+	}
+
 	for (const TPair< TSharedRef< IDatasmithLevelSequenceElement >, ULevelSequence* >& ImportedLevelSequencePair : ImportContext.ImportedLevelSequences)
 	{
 		if (ImportContext.bUserCancelled)
diff --git a/Engine/Plugins/Enterprise/DatasmithImporter/Source/DatasmithImporter/Private/DatasmithImporterModule.cpp b/Engine/Plugins/Enterprise/DatasmithImporter/Source/DatasmithImporter/Private/DatasmithImporterModule.cpp
--- a/Engine/Plugins/Enterprise/DatasmithImporter/Source/DatasmithImporter/Private/DatasmithImporterModule.cpp
+++ b/Engine/Plugins/Enterprise/DatasmithImporter/Source/DatasmithImporter/Private/DatasmithImporterModule.cpp
@@ -10,9 +10,9 @@
 #include "DatasmithContentEditorStyle.h"
 #include "DatasmithCustomAction.h"
 #include "DatasmithFileProducer.h"
-#include "DatasmithImportFactory.h"
 #include "DatasmithImporterEditorSettings.h"
 #include "DatasmithImporterHelper.h"
+#include "DatasmithImportFactory.h"
 #include "DatasmithScene.h"
 #include "DatasmithStaticMeshImporter.h"
 #include "DatasmithUtils.h"
@@ -60,6 +60,8 @@
 
 #define LOCTEXT_NAMESPACE "DatasmithImporter"
 
+static class IDatasmithImporterExt* GClothImporterExtInstance = nullptr;
+
 /**
  * DatasmithImporter module implementation (private)
  */
@@ -126,6 +128,18 @@ public:
 		return GetDefault< UDatasmithImporterEditorSettings >() && GetDefault< UDatasmithImporterEditorSettings >()->bOfflineImporter;
 	}
 
+
+	virtual void SetClothImporterExtension(class IDatasmithImporterExt* InInstance) override
+	{
+		GClothImporterExtInstance = InInstance;
+	}
+
+
+	virtual class IDatasmithImporterExt* GetClothImporterExtension() override
+	{
+		return GClothImporterExtInstance;
+	}
+
 private:
 	static TSharedRef<FExtender> OnExtendContentBrowserAssetSelectionMenu(const TArray<FAssetData>& SelectedAssets);
 	static TSharedRef<FExtender> OnExtendLevelEditorActorSelectionMenu(const TSharedRef<FUICommandList> CommandList, TArray<AActor*> SelectedActors);
diff --git a/Engine/Plugins/Enterprise/DatasmithImporter/Source/DatasmithImporter/Public/DatasmithClothImporter.h b/Engine/Plugins/Enterprise/DatasmithImporter/Source/DatasmithImporter/Public/DatasmithClothImporter.h
new file mode 100644
--- /dev/null
+++ b/Engine/Plugins/Enterprise/DatasmithImporter/Source/DatasmithImporter/Public/DatasmithClothImporter.h
@@ -0,0 +1,20 @@
+// Copyright Epic Games, Inc. All Rights Reserved.
+#pragma once
+
+
+#include "CoreTypes.h"
+#include "Templates/SharedPointer.h"
+#include "UObject/ObjectMacros.h"
+
+class FDatasmithCloth;
+class IDatasmithClothElement;
+
+
+// #ue_ds_cloth_arch: Temp API
+class IDatasmithImporterExt
+{
+public:
+	virtual UObject* MakeClothAsset(UObject* Outer, const TCHAR* Name, EObjectFlags ObjectFlags) = 0;
+	virtual void FillCloth(UObject* ClothAsset, TSharedRef<IDatasmithClothElement> ClothElement, FDatasmithCloth& DsCloth) = 0;
+};
+
diff --git a/Engine/Plugins/Enterprise/DatasmithImporter/Source/DatasmithImporter/Public/DatasmithImportContext.h b/Engine/Plugins/Enterprise/DatasmithImporter/Source/DatasmithImporter/Public/DatasmithImportContext.h
--- a/Engine/Plugins/Enterprise/DatasmithImporter/Source/DatasmithImporter/Public/DatasmithImportContext.h
+++ b/Engine/Plugins/Enterprise/DatasmithImporter/Source/DatasmithImporter/Public/DatasmithImportContext.h
@@ -21,6 +21,7 @@ struct FDatasmithImportContext;
 class FJsonObject;
 class FMessageLogModule;
 class IDatasmithBaseMaterialElement;
+class IDatasmithClothElement;
 class IDatasmithElement;
 class IDatasmithLevelSequenceElement;
 class IDatasmithLevelVariantSetsElement;
@@ -256,6 +257,8 @@ struct DATASMITHIMPORTER_API FDatasmithImportContext
 	/** Map of imported mesh for each mesh element */
 	TMap< TSharedRef< IDatasmithMeshElement >, UStaticMesh* > ImportedStaticMeshes;
 
+	TMap< TSharedRef< IDatasmithClothElement >, UObject* > ImportedClothes; // UChaosClothAsset
+
 	/** Register IDatasmithMeshElement by their name so they can be searched faster */
 	TMap< FString, TSharedRef < IDatasmithMeshElement > > ImportedStaticMeshesByName;
 
diff --git a/Engine/Plugins/Enterprise/DatasmithImporter/Source/DatasmithImporter/Public/DatasmithImporter.h b/Engine/Plugins/Enterprise/DatasmithImporter/Source/DatasmithImporter/Public/DatasmithImporter.h
--- a/Engine/Plugins/Enterprise/DatasmithImporter/Source/DatasmithImporter/Public/DatasmithImporter.h
+++ b/Engine/Plugins/Enterprise/DatasmithImporter/Source/DatasmithImporter/Public/DatasmithImporter.h
@@ -38,6 +38,8 @@ struct DATASMITHIMPORTER_API FDatasmithImporter
 	 */
 	static void ImportStaticMeshes( FDatasmithImportContext& ImportContext );
 
+	static void ImportClothes( FDatasmithImportContext& ImportContext );
+
 	/**
 	 * @return The imported Static Mesh
 	 */
@@ -55,6 +57,7 @@ struct DATASMITHIMPORTER_API FDatasmithImporter
 	 * @return The resulting static mesh
 	 */
 	static UStaticMesh* FinalizeStaticMesh( UStaticMesh* SourceStaticMesh, const TCHAR* StaticMeshesFolderPath, UStaticMesh* ExistingStaticMesh, TMap< UObject*, UObject* >* ReferencesToRemap = nullptr, bool bBuild = true );
+	static UObject* FinalizeCloth( UObject* SourceCloth, const TCHAR* FolderPath, UObject* ExistingCloth, TMap<UObject*, UObject*>* ReferencesToRemap = nullptr);
 
 	/**
 	 * Creates the AssetImportData associated with a UStaticMesh. Used when reimporting.
diff --git a/Engine/Plugins/Enterprise/DatasmithImporter/Source/DatasmithImporter/Public/DatasmithImporterModule.h b/Engine/Plugins/Enterprise/DatasmithImporter/Source/DatasmithImporter/Public/DatasmithImporterModule.h
--- a/Engine/Plugins/Enterprise/DatasmithImporter/Source/DatasmithImporter/Public/DatasmithImporterModule.h
+++ b/Engine/Plugins/Enterprise/DatasmithImporter/Source/DatasmithImporter/Public/DatasmithImporterModule.h
@@ -16,7 +16,6 @@
  */
 class IDatasmithImporterModule : public IModuleInterface
 {
-
 public:
 
 	/**
@@ -46,5 +45,9 @@ public:
 
 	DECLARE_MULTICAST_DELEGATE_OneParam(FOnGenerateDatasmithImportMenu, struct FToolMenuSection&);
 	virtual FOnGenerateDatasmithImportMenu& OnGenerateDatasmithImportMenu() = 0;
+
+	// #ue_ds_cloth_arch: temp API: Do not use directly.
+	virtual void SetClothImporterExtension(class IDatasmithImporterExt*) = 0;
+	virtual class IDatasmithImporterExt* GetClothImporterExtension() = 0;
 };
 
diff --git a/Engine/Plugins/Enterprise/DatasmithImporter/Source/DatasmithNativeTranslator/Private/DatasmithNativeTranslator.cpp b/Engine/Plugins/Enterprise/DatasmithImporter/Source/DatasmithNativeTranslator/Private/DatasmithNativeTranslator.cpp
--- a/Engine/Plugins/Enterprise/DatasmithImporter/Source/DatasmithNativeTranslator/Private/DatasmithNativeTranslator.cpp
+++ b/Engine/Plugins/Enterprise/DatasmithImporter/Source/DatasmithNativeTranslator/Private/DatasmithNativeTranslator.cpp
@@ -43,6 +43,13 @@ void FDatasmithNativeTranslator::ResolveSceneFilePaths(TSharedRef<IDatasmithScen
 		Mesh->SetFile(*ResolveFilePath(Path, ResourcePaths));
 	}
 
+	for (int32 Index = 0; Index < Scene->GetClothesCount(); ++Index)
+	{
+		const TSharedPtr<IDatasmithClothElement>& Cloth = Scene->GetCloth(Index);
+		const TCHAR* Path = Cloth->GetFile();
+		Cloth->SetFile(*ResolveFilePath(Path, ResourcePaths));
+	}
+
 	for (int32 Index = 0; Index < Scene->GetTexturesCount(); ++Index)
 	{
 		const TSharedPtr<IDatasmithTextureElement>& Tex = Scene->GetTexture(Index);
diff --git a/Engine/Plugins/Enterprise/DatasmithImporter/Source/DatasmithTranslator/Public/DatasmithPayload.h b/Engine/Plugins/Enterprise/DatasmithImporter/Source/DatasmithTranslator/Public/DatasmithPayload.h
--- a/Engine/Plugins/Enterprise/DatasmithImporter/Source/DatasmithTranslator/Public/DatasmithPayload.h
+++ b/Engine/Plugins/Enterprise/DatasmithImporter/Source/DatasmithTranslator/Public/DatasmithPayload.h
@@ -2,6 +2,7 @@
 
 #pragma once
 
+#include "DatasmithCloth.h"
 #include "MeshDescription.h"
 
 
@@ -10,6 +11,7 @@ struct FDatasmithPayload
 	TArray<class UDatasmithAdditionalData*> AdditionalData;
 };
 
+
 struct FDatasmithMeshElementPayload : public FDatasmithPayload
 {
 	TArray<FMeshDescription> LodMeshes;
@@ -17,6 +19,14 @@ struct FDatasmithMeshElementPayload : public FDatasmithPayload
 	TArray<FVector3f> CollisionPointCloud; // compatibility, favor the CollisionMesh member
 };
 
+
+// #ue_ds_cloth_doc IDatasmithScene API
+struct FDatasmithClothElementPayload : public FDatasmithPayload
+{
+	FDatasmithCloth Cloth;
+};
+
+
 struct FDatasmithLevelSequencePayload : public FDatasmithPayload
 {
 // #ueent_todo: split element in metadata/payload
diff --git a/Engine/Plugins/Enterprise/DatasmithImporter/Source/DatasmithTranslator/Public/DatasmithTranslator.h b/Engine/Plugins/Enterprise/DatasmithImporter/Source/DatasmithTranslator/Public/DatasmithTranslator.h
--- a/Engine/Plugins/Enterprise/DatasmithImporter/Source/DatasmithTranslator/Public/DatasmithTranslator.h
+++ b/Engine/Plugins/Enterprise/DatasmithImporter/Source/DatasmithTranslator/Public/DatasmithTranslator.h
@@ -13,6 +13,7 @@
 
 class FDatasmithSceneSource;
 class IDatasmithScene;
+class IDatasmithClothElement;
 class IDatasmithMeshElement;
 class IDatasmithLevelSequenceElement;
 
@@ -94,7 +95,16 @@ public:
 	 * @param OutMeshPayload	Actual mesh data from the source
 	 * @returns					Operation succeed
 	 */
-	virtual bool LoadStaticMesh(const TSharedRef< IDatasmithMeshElement > MeshElement, FDatasmithMeshElementPayload& OutMeshPayload) { return false; }
+	virtual bool LoadStaticMesh(const TSharedRef<IDatasmithMeshElement> MeshElement, FDatasmithMeshElementPayload& OutMeshPayload) { return false; }
+
+	/**
+	 * Get payload related to the given Element
+	 *
+	 * @param ClothElement       Element for which the payload is required
+	 * @param OutClothPayload    Actual cloth data from the source
+	 * @returns                  Operation succeed
+	 */
+	virtual bool LoadCloth(const TSharedRef<IDatasmithClothElement> ClothElement, FDatasmithClothElementPayload& OutClothPayload) { return false; }
 
 	/**
 	 * Get payload related to the given Element
@@ -103,7 +113,7 @@ public:
 	 * @param OutLevelSequencePayload	Data associated with this element
 	 * @returns							Operation succeed
 	 */
-	virtual bool LoadLevelSequence(const TSharedRef< IDatasmithLevelSequenceElement > LevelSequenceElement, FDatasmithLevelSequencePayload& OutLevelSequencePayload) { return false; }
+	virtual bool LoadLevelSequence(const TSharedRef<IDatasmithLevelSequenceElement> LevelSequenceElement, FDatasmithLevelSequencePayload& OutLevelSequencePayload) { return false; }
 
 	/**
 	 * Get the additional scene import options.
diff --git a/Engine/Source/Runtime/Datasmith/DatasmithCore/Private/DatasmithCloth.cpp b/Engine/Source/Runtime/Datasmith/DatasmithCore/Private/DatasmithCloth.cpp
new file mode 100644
--- /dev/null
+++ b/Engine/Source/Runtime/Datasmith/DatasmithCore/Private/DatasmithCloth.cpp
@@ -0,0 +1,27 @@
+// Copyright Epic Games, Inc. All Rights Reserved.
+
+#include "DatasmithCloth.h"
+
+// #include "DatasmithMesh.h"
+
+// class FDatasmithClothImpl
+// {
+// public:
+// 	FDatasmithMesh SimMesh;
+// 	FDatasmithMesh RenderMesh;
+// };
+//
+//
+// FDatasmithCloth::FDatasmithCloth()
+// 	: Impl(MakePimpl<FDatasmithClothImpl>())
+// {}
+//
+// FDatasmithMesh& FDatasmithCloth::GetSimMesh()
+// {
+// 	return Impl->SimMesh;
+// }
+//
+// FDatasmithMesh& FDatasmithCloth::GetRenderMesh()
+// {
+// 	return Impl->RenderMesh;
+// }
diff --git a/Engine/Source/Runtime/Datasmith/DatasmithCore/Private/DatasmithMeshSerialization.cpp b/Engine/Source/Runtime/Datasmith/DatasmithCore/Private/DatasmithMeshSerialization.cpp
--- a/Engine/Source/Runtime/Datasmith/DatasmithCore/Private/DatasmithMeshSerialization.cpp
+++ b/Engine/Source/Runtime/Datasmith/DatasmithCore/Private/DatasmithMeshSerialization.cpp
@@ -317,3 +317,26 @@ FDatasmithPackedMeshes GetDatasmithMeshFromFile(const FString& MeshPath)
 	return Result;
 }
 
+DATASMITHCORE_API FDatasmithPackedMeshes GetDatasmithClothFromFile(const FString& Path)
+{
+	// #ue_ds_cloth: asset serialization
+	FDatasmithPackedMeshes Result;
+
+	TUniquePtr<FArchive> Archive( IFileManager::Get().CreateFileReader(*Path) );
+	if ( !Archive.IsValid() )
+	{
+		UE_LOG(LogDatasmith, Warning, TEXT("Cannot read file %s"), *Path);
+		return Result;
+	}
+
+	Result.Serialize(*Archive);
+
+	if (Archive->IsError())
+	{
+		Result = FDatasmithPackedMeshes{};
+		UE_LOG(LogDatasmith, Warning, TEXT("Failed to read cloth from %s"), *Path);
+	}
+
+	return Result;
+}
+
diff --git a/Engine/Source/Runtime/Datasmith/DatasmithCore/Private/DatasmithSceneElementsImpl.cpp b/Engine/Source/Runtime/Datasmith/DatasmithCore/Private/DatasmithSceneElementsImpl.cpp
--- a/Engine/Source/Runtime/Datasmith/DatasmithCore/Private/DatasmithSceneElementsImpl.cpp
+++ b/Engine/Source/Runtime/Datasmith/DatasmithCore/Private/DatasmithSceneElementsImpl.cpp
@@ -86,7 +86,6 @@ int32 FDatasmithMeshElementImpl::GetMaterialSlotCount() const
 	return MaterialSlots.Num();
 }
 
-
 TSharedPtr<const IDatasmithMaterialIDElement> FDatasmithMeshElementImpl::GetMaterialSlotAt(int32 Index) const
 {
 	if (MaterialSlots.IsValidIndex(Index))
@@ -107,6 +106,12 @@ TSharedPtr<IDatasmithMaterialIDElement> FDatasmithMeshElementImpl::GetMaterialSl
 	return InvalidMaterialID;
 }
 
+FDatasmithClothElementImpl::FDatasmithClothElementImpl(const TCHAR* InName)
+	: FDatasmithElementImpl(InName, EDatasmithElementType::Cloth)
+{
+	Store.RegisterParameter(File, "File");
+}
+
 TSharedPtr< IDatasmithKeyValueProperty > FDatasmithKeyValuePropertyImpl::NullPropertyPtr;
 
 FDatasmithKeyValuePropertyImpl::FDatasmithKeyValuePropertyImpl(const TCHAR* InName)
@@ -1097,7 +1102,6 @@ void FDatasmithSceneImpl::SetHost(const TCHAR* InHostname)
 	Hostname = InHostname;
 }
 
-static const TSharedPtr< IDatasmithMeshElement > InvalidMeshElement;
 
 TSharedPtr< IDatasmithMeshElement > FDatasmithSceneImpl::GetMesh(int32 InIndex)
 {
@@ -1119,6 +1123,7 @@ const TSharedPtr< IDatasmithMeshElement >& FDatasmithSceneImpl::GetMesh(int32 In
 	}
 	else
 	{
+		static const TSharedPtr< IDatasmithMeshElement > InvalidMeshElement;
 		return InvalidMeshElement;
 	}
 }
@@ -1131,6 +1136,53 @@ void FDatasmithSceneImpl::RemoveMeshAt(int32 InIndex)
 	}
 }
 
+void FDatasmithSceneImpl::AddCloth(const TSharedPtr< IDatasmithClothElement >& InElement)
+{
+	Clothes.Add(InElement);
+}
+
+int32 FDatasmithSceneImpl::GetClothesCount() const
+{
+	return Clothes.Num();
+}
+
+TSharedPtr< IDatasmithClothElement > FDatasmithSceneImpl::GetCloth(int32 InIndex)
+{
+	if ( Clothes.IsValidIndex( InIndex ) )
+	{
+		return Clothes[InIndex];
+	}
+	return {};
+}
+
+const TSharedPtr< IDatasmithClothElement >& FDatasmithSceneImpl::GetCloth(int32 InIndex) const
+{
+	if ( Clothes.IsValidIndex( InIndex ) )
+	{
+		return Clothes[InIndex];
+	}
+	static TSharedPtr< IDatasmithClothElement > NullCloth;
+	return NullCloth;
+}
+
+void FDatasmithSceneImpl::RemoveCloth(const TSharedPtr< IDatasmithClothElement >& InElement)
+{
+	Clothes.Remove(InElement);
+}
+
+void FDatasmithSceneImpl::RemoveClothAt(int32 InIndex)
+{
+	if (Clothes.IsValidIndex(InIndex))
+	{
+		Clothes.RemoveAt(InIndex);
+	}
+}
+
+void FDatasmithSceneImpl::EmptyClothes()
+{
+	Clothes.Empty();
+}
+
 static const TSharedPtr< IDatasmithMetaDataElement > InvalidMetaData;
 
 TSharedPtr< IDatasmithMetaDataElement > FDatasmithSceneImpl::GetMetaData(int32 InIndex)
diff --git a/Engine/Source/Runtime/Datasmith/DatasmithCore/Private/DatasmithSceneElementsImpl.h b/Engine/Source/Runtime/Datasmith/DatasmithCore/Private/DatasmithSceneElementsImpl.h
--- a/Engine/Source/Runtime/Datasmith/DatasmithCore/Private/DatasmithSceneElementsImpl.h
+++ b/Engine/Source/Runtime/Datasmith/DatasmithCore/Private/DatasmithSceneElementsImpl.h
@@ -387,6 +387,21 @@ private:
 	TDatasmithReferenceArrayProxy<IDatasmithMaterialIDElement> MaterialSlots;
 };
 
+// #ue_ds_cloth doc: Experimental Element that describes a cloth asset
+class FDatasmithClothElementImpl : public FDatasmithElementImpl< IDatasmithClothElement >
+{
+public:
+	explicit FDatasmithClothElementImpl(const TCHAR* InName);
+
+public:
+	virtual const TCHAR* GetFile() const override { return *(FString&)File; }
+	virtual void SetFile(const TCHAR* InFile) override { File = InFile; }
+
+private:
+	TReflected<FString> File;
+};
+
+
 class FDatasmithMaterialIDElementImpl : public FDatasmithElementImpl< IDatasmithMaterialIDElement >
 {
 public:
@@ -1589,6 +1604,14 @@ public:
 	virtual void RemoveMeshAt(int32 InIndex) override;
 	virtual void EmptyMeshes() override { Meshes.Empty(); }
 
+	virtual void AddCloth(const TSharedPtr< IDatasmithClothElement >& InElement) override;
+	virtual int32 GetClothesCount() const override;
+	virtual TSharedPtr< IDatasmithClothElement > GetCloth(int32 InIndex) override;
+	virtual const TSharedPtr< IDatasmithClothElement >& GetCloth(int32 InIndex) const override;
+	virtual void RemoveCloth(const TSharedPtr< IDatasmithClothElement >& InElement) override;
+	virtual void RemoveClothAt(int32 InIndex) override;
+	virtual void EmptyClothes() override;
+
 	virtual void AddActor(const TSharedPtr< IDatasmithActorElement >& InActor) override { Actors.Add(InActor);  }
 	virtual int32 GetActorsCount() const override { return Actors.Num(); }
 	virtual TSharedPtr< IDatasmithActorElement > GetActor(int32 InIndex) override;
@@ -1649,6 +1672,7 @@ private:
 
 	TDatasmithReferenceArrayProxy<IDatasmithActorElement>            Actors;
 	TDatasmithReferenceArrayProxy<IDatasmithMeshElement>             Meshes;
+	TDatasmithReferenceArrayProxy<IDatasmithClothElement>            Clothes;
 	TDatasmithReferenceArrayProxy<IDatasmithBaseMaterialElement>     Materials;
 	TDatasmithReferenceArrayProxy<IDatasmithTextureElement>          Textures;
 	TDatasmithReferenceArrayProxy<IDatasmithMetaDataElement>         MetaData;
@@ -1670,7 +1694,7 @@ private:
 
 	TReflected<bool> bUseSky;
 
-	// Internal cache for faster metadata access per-element, should be accessed via GetElementToMetaDataCache(), do not use directly.
+	// Internal cache for faster metadata access per-element, should be accessed via GetMetaDataCache(), do not use directly.
 	mutable TMap< TSharedPtr< IDatasmithElement >, TSharedPtr< IDatasmithMetaDataElement> > ElementToMetaDataMap;
 	TMap< TSharedPtr< IDatasmithElement >, TSharedPtr< IDatasmithMetaDataElement> >& GetElementToMetaDataCache() const;
 };
diff --git a/Engine/Source/Runtime/Datasmith/DatasmithCore/Private/DatasmithSceneFactory.cpp b/Engine/Source/Runtime/Datasmith/DatasmithCore/Private/DatasmithSceneFactory.cpp
--- a/Engine/Source/Runtime/Datasmith/DatasmithCore/Private/DatasmithSceneFactory.cpp
+++ b/Engine/Source/Runtime/Datasmith/DatasmithCore/Private/DatasmithSceneFactory.cpp
@@ -40,6 +40,8 @@ TSharedPtr< IDatasmithElement > FDatasmithSceneFactory::CreateElement( EDatasmit
 		}
 	case EDatasmithElementType::StaticMesh:
 		return CreateMesh( InName );
+	case EDatasmithElementType::Cloth:
+		return CreateCloth( InName );
 	case EDatasmithElementType::StaticMeshActor:
 		return CreateMeshActor( InName );
 	case EDatasmithElementType::PointLight:
@@ -194,6 +196,11 @@ TSharedRef< IDatasmithMeshElement > FDatasmithSceneFactory::CreateMesh( const TC
 	return MakeShared< FDatasmithMeshElementImpl >( InName );
 }
 
+TSharedRef< IDatasmithClothElement > FDatasmithSceneFactory::CreateCloth(const TCHAR* InName)
+{
+	return MakeShared< FDatasmithClothElementImpl >( InName );
+}
+
 TSharedRef< IDatasmithMeshActorElement > FDatasmithSceneFactory::CreateMeshActor( const TCHAR* InName )
 {
 	return MakeShared< FDatasmithMeshActorElementImpl<> >( InName );
diff --git a/Engine/Source/Runtime/Datasmith/DatasmithCore/Private/DatasmithSceneXmlReader.cpp b/Engine/Source/Runtime/Datasmith/DatasmithCore/Private/DatasmithSceneXmlReader.cpp
--- a/Engine/Source/Runtime/Datasmith/DatasmithCore/Private/DatasmithSceneXmlReader.cpp
+++ b/Engine/Source/Runtime/Datasmith/DatasmithCore/Private/DatasmithSceneXmlReader.cpp
@@ -435,6 +435,20 @@ void FDatasmithSceneXmlReader::ParseMesh(FXmlNode* InNode, TSharedPtr<IDatasmith
 	}
 }
 
+void FDatasmithSceneXmlReader::ParseCloth(FXmlNode* InNode, TSharedPtr<IDatasmithClothElement>& OutElement) const
+{
+	// #ue_ds_cloth: xml read
+	ParseElement( InNode, OutElement.ToSharedRef() );
+
+	for (const FXmlNode* Node : InNode->GetChildrenNodes())
+	{
+		if (Node->GetTag() == TEXT("file"))
+		{
+			OutElement->SetFile(*UnsanitizeXMLText(Node->GetAttribute(TEXT("path"))));
+		}
+	}
+}
+
 void FDatasmithSceneXmlReader::ParseTextureElement(FXmlNode* InNode, TSharedPtr<IDatasmithTextureElement>& OutElement) const
 {
 	TRACE_CPUPROFILER_EVENT_SCOPE(FDatasmithSceneXmlReader::ParseTextureElement);
@@ -1191,6 +1205,18 @@ bool FDatasmithSceneXmlReader::ParseXmlFile(TSharedRef< IDatasmithScene >& OutSc
 
 			Objects.Add( Element->GetName(), Element );
 		}
+		// CLOTHES
+		else if (Nodes[i]->GetTag() == DATASMITH_CLOTH)
+		{
+			FString ElementName = Nodes[i]->GetAttribute(TEXT("name"));
+			TSharedPtr< IDatasmithClothElement > Element = FDatasmithSceneFactory::CreateCloth(*ElementName);
+
+			ParseCloth( Nodes[i], Element );
+
+			OutScene->AddCloth(Element);
+
+// 			Objects.Add( Element->GetName(), Element ); // #ue_ds_cloth_todo cloths referencables by other elements
+		}
 		// LEVEL SEQUENCES
 		else if (Nodes[i]->GetTag() == DATASMITH_LEVELSEQUENCENAME)
 		{
diff --git a/Engine/Source/Runtime/Datasmith/DatasmithCore/Private/DatasmithSceneXmlWriter.cpp b/Engine/Source/Runtime/Datasmith/DatasmithCore/Private/DatasmithSceneXmlWriter.cpp
--- a/Engine/Source/Runtime/Datasmith/DatasmithCore/Private/DatasmithSceneXmlWriter.cpp
+++ b/Engine/Source/Runtime/Datasmith/DatasmithCore/Private/DatasmithSceneXmlWriter.cpp
@@ -30,6 +30,7 @@ public:
 
 	static void WriteTexture(FArchive& Archive, int32 Indent, const TCHAR* Prefix, const TCHAR* Name, FDatasmithTextureSampler UV);
 	static void WriteMeshElement(const TSharedPtr< IDatasmithMeshElement >& MeshElement, FArchive& Archive, int32 Indent);
+	static void WriteClothElement(const TSharedPtr< IDatasmithClothElement >& ClothElement, FArchive& Archive, int32 Indent);
 	static void WriteLevelSequenceElement( const TSharedRef< IDatasmithLevelSequenceElement>& SequenceElement, FArchive& Archive, int32 Indent );
 
 	static void WriteLevelVariantSetsElement( const TSharedRef< IDatasmithLevelVariantSetsElement >& LevelVariantSetsElement, FArchive& Archive, int32 Indent );
@@ -356,6 +357,23 @@ void FDatasmithSceneXmlWriterImpl::WriteMeshElement(const TSharedPtr< IDatasmith
 	SerializeToArchive(Archive, XmlString);
 }
 
+void FDatasmithSceneXmlWriterImpl::WriteClothElement(const TSharedPtr<IDatasmithClothElement>& Element, FArchive& Archive, int32 Indent)
+{
+	// #ue_ds_cloth_note: FDatasmithSceneXmlWriterImpl::WriteClothElement
+	WriteIndent(Archive, Indent);
+	FString XmlString = TEXT("<") DATASMITH_CLOTH;
+	XmlString += TEXT(" name=\"") + SanitizeXMLText( Element->GetName() ) + TEXT("\"");
+	XmlString += TEXT(" label=\"") + SanitizeXMLText( Element->GetLabel() ) + TEXT("\"");
+	XmlString += TEXT(">") LINE_TERMINATOR;
+	SerializeToArchive(Archive, XmlString);
+
+	WriteIndent(Archive, Indent + 1);
+	SerializeToArchive(Archive, TEXT("<file path=\"") + SanitizeXMLText(Element->GetFile()) + TEXT("\"/>") LINE_TERMINATOR);
+
+	WriteIndent(Archive, Indent);
+	SerializeToArchive(Archive, TEXT("</") DATASMITH_CLOTH TEXT(">") LINE_TERMINATOR);
+}
+
 void FDatasmithSceneXmlWriterImpl::WriteLevelSequenceElement(const TSharedRef< IDatasmithLevelSequenceElement>& SequenceElement, FArchive& Archive, int32 Indent)
 {
 	WriteIndent(Archive, Indent);
@@ -1922,6 +1940,12 @@ void FDatasmithSceneXmlWriter::Serialize( TSharedRef< IDatasmithScene > Datasmit
 		FDatasmithSceneXmlWriterImpl::WriteMeshElement( Mesh, Archive, 1 );
 	}
 
+	for ( int32 Index = 0; Index < DatasmithScene->GetClothesCount(); ++Index )
+	{
+		const TSharedPtr< IDatasmithClothElement >& Cloth = DatasmithScene->GetCloth( Index );
+		FDatasmithSceneXmlWriterImpl::WriteClothElement( Cloth, Archive, 1 );
+	}
+
 	for ( int32 ActorIndex = 0; ActorIndex < DatasmithScene->GetActorsCount(); ++ActorIndex )
 	{
 		TSharedPtr< IDatasmithActorElement > Actor = DatasmithScene->GetActor( ActorIndex );
diff --git a/Engine/Source/Runtime/Datasmith/DatasmithCore/Public/DatasmithCloth.h b/Engine/Source/Runtime/Datasmith/DatasmithCore/Public/DatasmithCloth.h
new file mode 100644
--- /dev/null
+++ b/Engine/Source/Runtime/Datasmith/DatasmithCore/Public/DatasmithCloth.h
@@ -0,0 +1,31 @@
+// Copyright Epic Games, Inc. All Rights Reserved.
+#pragma once
+
+#include "CoreMinimal.h"
+// #include "Templates/PimplPtr.h"
+// #include "Containers/Array.h"
+
+class FDatasmithMesh;
+
+class DATASMITHCORE_API FDatasmithClothPattern
+{
+public:
+	TArray<FVector2f> SimPosition;
+	TArray<FVector3f> SimRestPosition;
+	TArray<uint32> SimTriangleIndices;
+};
+
+/*
+ * Structure of a Cloth
+ * - One Thin Mesh:
+ * - base for physics simulation
+ *
+ * - One Thick Mesh
+ * - aka. render mesh
+ */
+class DATASMITHCORE_API FDatasmithCloth
+{
+public:
+	TArray<FDatasmithClothPattern> Patterns;
+};
+
diff --git a/Engine/Source/Runtime/Datasmith/DatasmithCore/Public/DatasmithDefinitions.h b/Engine/Source/Runtime/Datasmith/DatasmithCore/Public/DatasmithDefinitions.h
--- a/Engine/Source/Runtime/Datasmith/DatasmithCore/Public/DatasmithDefinitions.h
+++ b/Engine/Source/Runtime/Datasmith/DatasmithCore/Public/DatasmithDefinitions.h
@@ -48,6 +48,7 @@ enum class EDatasmithElementType : uint64
 	MaterialExpression             = 1ull << 31,
 	MaterialExpressionInput        = 1ull << 32,
 	MaterialExpressionOutput       = 1ull << 33,
+	Cloth                          = 1ull << 34,
 };
 
 ENUM_CLASS_FLAGS( EDatasmithElementType ); // Define bitwise operators for EDatasmithElementType
@@ -454,6 +455,8 @@ static const TCHAR* KeyValuePropertyTypeStrings[] = { TEXT("String"), TEXT("Colo
 
 #define DATASMITH_ACTORNAME						TEXT("Actor")
 
+#define DATASMITH_CLOTH							TEXT("Cloth")
+
 //ACTOR MESHES
 #define DATASMITH_ACTORMESHNAME					TEXT("ActorMesh")
 
@@ -628,3 +631,4 @@ static const TCHAR* KeyValuePropertyTypeStrings[] = { TEXT("String"), TEXT("Colo
 #define DATASMITH_ENVILLUMINATIONMAP			TEXT("Illuminate")
 
 #define DATASMITH_DECALMATERIALNAME				TEXT("DecalMaterial")
+
diff --git a/Engine/Source/Runtime/Datasmith/DatasmithCore/Public/DatasmithMeshSerialization.h b/Engine/Source/Runtime/Datasmith/DatasmithCore/Public/DatasmithMeshSerialization.h
--- a/Engine/Source/Runtime/Datasmith/DatasmithCore/Public/DatasmithMeshSerialization.h
+++ b/Engine/Source/Runtime/Datasmith/DatasmithCore/Public/DatasmithMeshSerialization.h
@@ -24,3 +24,5 @@ struct DATASMITHCORE_API FDatasmithPackedMeshes
 };
 
 DATASMITHCORE_API FDatasmithPackedMeshes GetDatasmithMeshFromFile(const FString& MeshPath);
+DATASMITHCORE_API FDatasmithPackedMeshes GetDatasmithClothFromFile(const FString& Path);
+
diff --git a/Engine/Source/Runtime/Datasmith/DatasmithCore/Public/DatasmithSceneFactory.h b/Engine/Source/Runtime/Datasmith/DatasmithCore/Public/DatasmithSceneFactory.h
--- a/Engine/Source/Runtime/Datasmith/DatasmithCore/Public/DatasmithSceneFactory.h
+++ b/Engine/Source/Runtime/Datasmith/DatasmithCore/Public/DatasmithSceneFactory.h
@@ -44,6 +44,8 @@ public:
 
 	static TSharedRef< IDatasmithMeshElement > CreateMesh( const TCHAR* InName );
 
+	static TSharedRef< IDatasmithClothElement > CreateCloth( const TCHAR* InName );
+
 	static TSharedRef< IDatasmithMeshActorElement > CreateMeshActor( const TCHAR* InName );
 
 	static TSharedRef< IDatasmithHierarchicalInstancedStaticMeshActorElement > CreateHierarchicalInstanceStaticMeshActor( const TCHAR* InName );
diff --git a/Engine/Source/Runtime/Datasmith/DatasmithCore/Public/DatasmithSceneXmlReader.h b/Engine/Source/Runtime/Datasmith/DatasmithCore/Public/DatasmithSceneXmlReader.h
--- a/Engine/Source/Runtime/Datasmith/DatasmithCore/Public/DatasmithSceneXmlReader.h
+++ b/Engine/Source/Runtime/Datasmith/DatasmithCore/Public/DatasmithSceneXmlReader.h
@@ -42,6 +42,7 @@ private:
 	void ParsePropertyCapture( FXmlNode* InNode, const TSharedRef<IDatasmithPropertyCaptureElement>& OutElement ) const;
 	void ParseObjectPropertyCapture( FXmlNode* InNode, const TSharedRef<IDatasmithObjectPropertyCaptureElement>& OutElement, const TMap< FString, TSharedPtr<IDatasmithElement> >& Objects ) const;
 	void ParseMesh(FXmlNode* InNode, TSharedPtr<IDatasmithMeshElement>& OutElement) const;
+	void ParseCloth(FXmlNode* InNode, TSharedPtr<IDatasmithClothElement>& OutElement) const;
 	void ParseTextureElement(FXmlNode* InNode, TSharedPtr<IDatasmithTextureElement>& OutElement) const;
 	void ParseTexture(FXmlNode* InNode, FString& OutTextureFilename, FDatasmithTextureSampler& OutTextureUV) const;
 	void ParseActor(FXmlNode* InNode, TSharedPtr<IDatasmithActorElement>& InOutElement, TSharedRef< IDatasmithScene > Scene, TMap< FString, TSharedPtr<IDatasmithActorElement> >& Actors) const;
diff --git a/Engine/Source/Runtime/Datasmith/DatasmithCore/Public/IDatasmithSceneElements.h b/Engine/Source/Runtime/Datasmith/DatasmithCore/Public/IDatasmithSceneElements.h
--- a/Engine/Source/Runtime/Datasmith/DatasmithCore/Public/IDatasmithSceneElements.h
+++ b/Engine/Source/Runtime/Datasmith/DatasmithCore/Public/IDatasmithSceneElements.h
@@ -24,8 +24,6 @@ class IDatasmithLevelVariantSetsElement;
 class DATASMITHCORE_API IDatasmithElement : public DirectLink::ISceneGraphNode
 {
 public:
-	virtual ~IDatasmithElement() {}
-
 	/** returns if this DatasmithElement is of a specified type */
 	virtual bool IsA(EDatasmithElementType Type) const = 0;
 
@@ -238,6 +236,21 @@ protected:
 	friend class FDatasmithStaticMeshImporter;
 };
 
+/**
+* #ue_ds_cloth_doc IDatasmithClothElement class: experimental class that describes a cloth asset
+*/
+class DATASMITHCORE_API IDatasmithClothElement : public IDatasmithElement
+{
+public:
+	/** Get the FDatasmithCloth resource filename */
+	virtual const TCHAR* GetFile() const = 0;
+
+	/** Set the FDatasmithCloth resource filename, it can be absolute or relative to the scene file */
+	virtual void SetFile(const TCHAR* InFile) = 0;
+
+// class DATASMITHCORE_API IDatasmithClothPropertiesElement : public IDatasmithElement
+};
+
 /**
  * IDatasmithActorElement used in any geometry instance independently if it could be static or movable.
  * It doesn't define the actual geometry, you'll need IDatasmithMeshElement for this.
@@ -678,8 +691,6 @@ public:
 class DATASMITHCORE_API IDatasmithPostProcessElement : public IDatasmithElement
 {
 public:
-	virtual ~IDatasmithPostProcessElement() {}
-
 	/** Get color filter temperature in Kelvin */
 	virtual float GetTemperature() const = 0;
 
@@ -1591,6 +1602,15 @@ public:
 	*/
 	virtual void EmptyMeshes() = 0;
 
+	// #ue_ds_cloth_doc IDatasmithScene API
+	virtual void AddCloth(const TSharedPtr< IDatasmithClothElement >& InElement) = 0;
+	virtual int32 GetClothesCount() const = 0;
+	virtual TSharedPtr< IDatasmithClothElement > GetCloth(int32 InIndex) = 0;
+	virtual const TSharedPtr< IDatasmithClothElement >& GetCloth(int32 InIndex) const = 0;
+	virtual void RemoveCloth(const TSharedPtr< IDatasmithClothElement >& InElement) = 0;
+	virtual void RemoveClothAt(int32 InIndex) = 0;
+	virtual void EmptyClothes() = 0;
+
 	/**
 	 * Adds an Actor to the scene.
 	 *
