Dataprep: Cannot commit dataprep with a rect light in the scene

diff --git a/Engine/Plugins/Enterprise/DatasmithImporter/Source/DatasmithImporter/Private/DatasmithConsumer.cpp b/Engine/Plugins/Enterprise/DatasmithImporter/Source/DatasmithImporter/Private/DatasmithConsumer.cpp
--- a/Engine/Plugins/Enterprise/DatasmithImporter/Source/DatasmithImporter/Private/DatasmithConsumer.cpp
+++ b/Engine/Plugins/Enterprise/DatasmithImporter/Source/DatasmithImporter/Private/DatasmithConsumer.cpp
@@ -43,6 +43,7 @@
 #include "Engine/Level.h"
 #include "Engine/LevelStreamingAlwaysLoaded.h"
 #include "Engine/Light.h"
+#include "Engine/RectLight.h"
 #include "Engine/Selection.h"
 #include "Engine/StaticMesh.h"
 #include "Engine/StaticMeshActor.h"
@@ -1567,7 +1568,7 @@ namespace DatasmithConsumerUtils
 		auto IsUnregisteredActor = [&](AActor* Actor)
 		{
 			// Skip non-imported actors
-			if( Actor == RootSceneActor || Actor == nullptr || Actor->GetRootComponent() == nullptr || Actor->IsA<AWorldSettings>() || Actor->IsA<APhysicsVolume>() || Actor->IsA<ABrush>() )
+			if( Actor == RootSceneActor || Actor == nullptr || Actor->GetRootComponent() == nullptr || Actor->IsA<AWorldSettings>() || Actor->IsA<APhysicsVolume>() || Actor->IsA<ABrush>() || Actor->IsA<ARectLight>())
 			{
 				return false;
 			}
diff --git a/Engine/Plugins/Enterprise/DatasmithImporter/Source/DatasmithImporter/Private/Utility/DatasmithImporterUtils.cpp b/Engine/Plugins/Enterprise/DatasmithImporter/Source/DatasmithImporter/Private/Utility/DatasmithImporterUtils.cpp
--- a/Engine/Plugins/Enterprise/DatasmithImporter/Source/DatasmithImporter/Private/Utility/DatasmithImporterUtils.cpp
+++ b/Engine/Plugins/Enterprise/DatasmithImporter/Source/DatasmithImporter/Private/Utility/DatasmithImporterUtils.cpp
@@ -39,6 +39,7 @@
 #include "Components/SpotLightComponent.h"
 #include "Editor.h"
 #include "Engine/DirectionalLight.h"
+#include "Engine/RectLight.h"
 #include "Engine/Level.h"
 #include "Engine/Light.h"
 #include "Engine/PostProcessVolume.h"
@@ -736,6 +737,11 @@ namespace FDatasmithImporterUtilsHelper
 
 			LightActorElement = SpotLightActorElement;
 		}
+		else if(LightActor->IsA<ARectLight>())
+		{
+			// Rect light is not yet supported
+			return TSharedPtr< IDatasmithLightActorElement >();
+		}
 		else
 		{
 			TSharedPtr< IDatasmithPointLightElement > PointLightActorElement = FDatasmithSceneFactory::CreatePointLight( *LightActor->GetName() );
@@ -1035,6 +1041,11 @@ namespace FDatasmithImporterUtilsHelper
 			}
 		}
 
+		if (!ActorElement.IsValid())
+		{
+			return TSharedPtr< IDatasmithActorElement >();
+		}
+
 		// Store actor's label
 		ActorElement->SetLabel( *Actor->GetActorLabel() );
 
@@ -1068,6 +1079,11 @@ namespace FDatasmithImporterUtilsHelper
 		{
 			TSharedPtr< IDatasmithActorElement > ChildActorElement = ConvertActorToActorElement( ChildActor, SceneElement );
 
+			if (!ChildActorElement.IsValid())
+			{
+				continue;
+			}
+
 			ParentActorElement->AddChild( ChildActorElement );
 
 			TArray<AActor*> ActorsToVisit;
@@ -1085,6 +1101,11 @@ void FDatasmithImporterUtils::FillSceneElement(TSharedPtr<IDatasmithScene>& Scen
 		// Convert root actor to actor element
 		TSharedPtr< IDatasmithActorElement > RootActorElement = FDatasmithImporterUtilsHelper::ConvertActorToActorElement( RootActor, SceneElement );
 
+		if (!RootActorElement.IsValid())
+		{
+			continue;
+		}
+
 		// Add newly created actor element to scene element
 		SceneElement->AddActor( RootActorElement );
 
