Cannot commit dataprep with a rect light in the scene

Added warnings

diff --git a/Engine/Plugins/Enterprise/DatasmithImporter/Source/DatasmithImporter/Private/DatasmithConsumer.cpp b/Engine/Plugins/Enterprise/DatasmithImporter/Source/DatasmithImporter/Private/DatasmithConsumer.cpp
--- a/Engine/Plugins/Enterprise/DatasmithImporter/Source/DatasmithImporter/Private/DatasmithConsumer.cpp
+++ b/Engine/Plugins/Enterprise/DatasmithImporter/Source/DatasmithImporter/Private/DatasmithConsumer.cpp
@@ -1570,6 +1570,10 @@ namespace DatasmithConsumerUtils
 			// Skip non-imported actors
 			if( Actor == RootSceneActor || Actor == nullptr || Actor->GetRootComponent() == nullptr || Actor->IsA<AWorldSettings>() || Actor->IsA<APhysicsVolume>() || Actor->IsA<ABrush>() || Actor->IsA<ARectLight>())
 			{
+				if (Actor != nullptr && Actor->IsA<ARectLight>())
+				{
+					UE_LOG(LogDatasmithImport, Warning, TEXT("Import of RectLight is not yet supported"));
+				}
 				return false;
 			}
 
diff --git a/Engine/Plugins/Enterprise/DatasmithImporter/Source/DatasmithImporter/Private/Utility/DatasmithImporterUtils.cpp b/Engine/Plugins/Enterprise/DatasmithImporter/Source/DatasmithImporter/Private/Utility/DatasmithImporterUtils.cpp
--- a/Engine/Plugins/Enterprise/DatasmithImporter/Source/DatasmithImporter/Private/Utility/DatasmithImporterUtils.cpp
+++ b/Engine/Plugins/Enterprise/DatasmithImporter/Source/DatasmithImporter/Private/Utility/DatasmithImporterUtils.cpp
@@ -739,7 +739,7 @@ namespace FDatasmithImporterUtilsHelper
 		}
 		else if(LightActor->IsA<ARectLight>())
 		{
-			// Rect light is not yet supported
+			UE_LOG(LogDatasmithImport, Warning, TEXT("Import of RectLight is not yet supported"));
 			return TSharedPtr< IDatasmithLightActorElement >();
 		}
 		else
