Datasmith: minor namespaces cleanup

diff --git a/Engine/Plugins/Enterprise/DatasmithImporter/Source/DatasmithImporter/Private/DatasmithImporter.cpp b/Engine/Plugins/Enterprise/DatasmithImporter/Source/DatasmithImporter/Private/DatasmithImporter.cpp
--- a/Engine/Plugins/Enterprise/DatasmithImporter/Source/DatasmithImporter/Private/DatasmithImporter.cpp
+++ b/Engine/Plugins/Enterprise/DatasmithImporter/Source/DatasmithImporter/Private/DatasmithImporter.cpp
@@ -82,170 +82,161 @@ extern UNREALED_API UEditorEngine* GEditor;
 
 #define LOCTEXT_NAMESPACE "DatasmithImporter"
 
-namespace UE
+namespace UE::DatasmithImporter::Private
 {
-	namespace DatasmithImporter
+	/**
+	 * Make the necessary adjustment to the materials so that the functions have unique parameter names in the material where they are use.
+	 */
+	void MakeChangeToExpressionsNameForMaterialFunctions(const TArray<FDatasmithImporterUtils::FFunctionAndMaterialsThatUseIt>& FunctionMaterialAndTheirReferencer)
 	{
-		namespace Private
-		{
-			namespace DatasmithImporter
-			{
-				/**
-				 * Make the necessary adjustment to the materials so that the functions have unique parameter names in the material where they are use.
-				 */
-				void MakeChangeToExpressionsNameForMaterialFunctions(const TArray<FDatasmithImporterUtils::FFunctionAndMaterialsThatUseIt>& FunctionMaterialAndTheirReferencer)
-				{
-					// Contains the name of the parameters for a material that isn't use as a function and the set of names used the by material functions it call
-					TMap<TSharedPtr<IDatasmithUEPbrMaterialElement>,TArray<TSharedRef<TSet<TPair<FName, EDatasmithMaterialExpressionType>>>>> TopMaterialAndNamesUsedByType;
-					TMap<EDatasmithMaterialExpressionType, TUniquePtr<FDatasmithUniqueNameProvider>> ExpressionTypeToUniqueNameProvider;
-					using FFunctionAndMaterialsThatUseIt = FDatasmithImporterUtils::FFunctionAndMaterialsThatUseIt;
+		// Contains the name of the parameters for a material that isn't use as a function and the set of names used the by material functions it call
+		TMap<TSharedPtr<IDatasmithUEPbrMaterialElement>,TArray<TSharedRef<TSet<TPair<FName, EDatasmithMaterialExpressionType>>>>> TopMaterialAndNamesUsedByType;
+		TMap<EDatasmithMaterialExpressionType, TUniquePtr<FDatasmithUniqueNameProvider>> ExpressionTypeToUniqueNameProvider;
+		using FFunctionAndMaterialsThatUseIt = FDatasmithImporterUtils::FFunctionAndMaterialsThatUseIt;
 
-					TMap<TSharedPtr<IDatasmithUEPbrMaterialElement>, int32> MaterialToIndexInArray;
-					MaterialToIndexInArray.Reserve( FunctionMaterialAndTheirReferencer.Num() );
+		TMap<TSharedPtr<IDatasmithUEPbrMaterialElement>, int32> MaterialToIndexInArray;
+		MaterialToIndexInArray.Reserve( FunctionMaterialAndTheirReferencer.Num() );
 
-					// Check if the name was already used in this function or in one of its referencers
-					TFunction<bool (uint32, const TSharedPtr<IDatasmithUEPbrMaterialElement>&, uint32, FName, EDatasmithMaterialExpressionType)> CheckIfNameIsUsed;
-					CheckIfNameIsUsed = [&TopMaterialAndNamesUsedByType, &MaterialToIndexInArray, &FunctionMaterialAndTheirReferencer, &CheckIfNameIsUsed] (uint32 MaterialHash, const TSharedPtr<IDatasmithUEPbrMaterialElement>& Material, uint32 ExpressionHash, FName Expression, EDatasmithMaterialExpressionType ExpressionType)
+		// Check if the name was already used in this function or in one of its referencers
+		TFunction<bool (uint32, const TSharedPtr<IDatasmithUEPbrMaterialElement>&, uint32, FName, EDatasmithMaterialExpressionType)> CheckIfNameIsUsed;
+		CheckIfNameIsUsed = [&TopMaterialAndNamesUsedByType, &MaterialToIndexInArray, &FunctionMaterialAndTheirReferencer, &CheckIfNameIsUsed] (uint32 MaterialHash, const TSharedPtr<IDatasmithUEPbrMaterialElement>& Material, uint32 ExpressionHash, FName Expression, EDatasmithMaterialExpressionType ExpressionType)
+			{
+				if ( int32* Index = MaterialToIndexInArray.FindByHash( MaterialHash, Material ) )
+				{
+					for ( const TSharedPtr<IDatasmithUEPbrMaterialElement>& Referencer : FunctionMaterialAndTheirReferencer[*Index].Value )
+					{
+						if ( CheckIfNameIsUsed( GetTypeHash( Referencer ), Referencer, ExpressionHash, Expression, ExpressionType ) )
 						{
-							if ( int32* Index = MaterialToIndexInArray.FindByHash( MaterialHash, Material ) )
-							{
-								for ( const TSharedPtr<IDatasmithUEPbrMaterialElement>& Referencer : FunctionMaterialAndTheirReferencer[*Index].Value )
-								{
-									if ( CheckIfNameIsUsed( GetTypeHash( Referencer ), Referencer, ExpressionHash, Expression, ExpressionType ) )
-									{
-										return true;
-									}
-								}
-							}
-							else if ( const TArray<TSharedRef<TSet<TPair<FName, EDatasmithMaterialExpressionType>>>>* NameUsedArray = TopMaterialAndNamesUsedByType.FindByHash( MaterialHash, Material ) )
-							// Only check if we are in a material that aren't referenced elsewhere
-							{
-								for ( const TSharedRef<TSet<TPair<FName, EDatasmithMaterialExpressionType>>>& NamesUsed : *NameUsedArray )
-								{
-									if ( NamesUsed->Contains( TPair<FName, EDatasmithMaterialExpressionType>( Expression, ExpressionType ) ) )
-									{
-										return true;
-									}
-								}
-							}
-
-							return false;
-						};
-
-					// Make sure that the parameter have a unique name in the materials were they are use
-					auto HandleExpressionNameForMaterialFunction =
-						[&ExpressionTypeToUniqueNameProvider, &CheckIfNameIsUsed] (uint32 MaterialHash, const TSharedPtr<IDatasmithUEPbrMaterialElement>& Material, FName Expression, EDatasmithMaterialExpressionType ExpressionType, int32 ExpressionIndex, TSet<TPair<FName, EDatasmithMaterialExpressionType>>& NameUsed)
+							return true;
+						}
+					}
+				}
+				else if ( const TArray<TSharedRef<TSet<TPair<FName, EDatasmithMaterialExpressionType>>>>* NameUsedArray = TopMaterialAndNamesUsedByType.FindByHash( MaterialHash, Material ) )
+				// Only check if we are in a material that aren't referenced elsewhere
+				{
+					for ( const TSharedRef<TSet<TPair<FName, EDatasmithMaterialExpressionType>>>& NamesUsed : *NameUsedArray )
+					{
+						if ( NamesUsed->Contains( TPair<FName, EDatasmithMaterialExpressionType>( Expression, ExpressionType ) ) )
 						{
-							TUniquePtr<FDatasmithUniqueNameProvider>& UniqueNameProvider = ExpressionTypeToUniqueNameProvider.FindOrAdd( ExpressionType );
-							if ( !UniqueNameProvider )
-							{
-								UniqueNameProvider = MakeUnique<FDatasmithUniqueNameProvider>();
-							}
+							return true;
+						}
+					}
+				}
 
-							FString ChosenName;
+				return false;
+			};
 
-							if ( CheckIfNameIsUsed( MaterialHash, Material, GetTypeHash( Expression ), Expression, ExpressionType ) )
-							{
-								// Make the name unique
-								ChosenName = UniqueNameProvider->GenerateUniqueName( Expression.ToString() );
-								Material->GetExpression( ExpressionIndex )->SetName( *ChosenName );
-							}
-							else
-							{
-								// Keep track of the used names
-								ChosenName = Expression.ToString();
-								UniqueNameProvider->AddExistingName( ChosenName );
-							}
+		// Make sure that the parameter have a unique name in the materials were they are use
+		auto HandleExpressionNameForMaterialFunction =
+			[&ExpressionTypeToUniqueNameProvider, &CheckIfNameIsUsed] (uint32 MaterialHash, const TSharedPtr<IDatasmithUEPbrMaterialElement>& Material, FName Expression, EDatasmithMaterialExpressionType ExpressionType, int32 ExpressionIndex, TSet<TPair<FName, EDatasmithMaterialExpressionType>>& NameUsed)
+			{
+				TUniquePtr<FDatasmithUniqueNameProvider>& UniqueNameProvider = ExpressionTypeToUniqueNameProvider.FindOrAdd( ExpressionType );
+				if ( !UniqueNameProvider )
+				{
+					UniqueNameProvider = MakeUnique<FDatasmithUniqueNameProvider>();
+				}
 
-							NameUsed.Add( TPair<FName, EDatasmithMaterialExpressionType>( *ChosenName, ExpressionType ) );
+				FString ChosenName;
 
-							if ( ExpressionType != EDatasmithMaterialExpressionType::Generic )
-							{
-								// Always add to the generic, we don't know what type of parameter a generic expression would be
-								TUniquePtr<FDatasmithUniqueNameProvider>& GenericUniqueNameProvider = ExpressionTypeToUniqueNameProvider.FindOrAdd( EDatasmithMaterialExpressionType::Generic );
-								if ( !GenericUniqueNameProvider )
-								{
-									GenericUniqueNameProvider = MakeUnique<FDatasmithUniqueNameProvider>();
-								}
-
-								GenericUniqueNameProvider->AddExistingName( ChosenName );
-								NameUsed.Add( TPair<FName, EDatasmithMaterialExpressionType>( *ChosenName, EDatasmithMaterialExpressionType::Generic ) );
-							}
-						};
+				if ( CheckIfNameIsUsed( MaterialHash, Material, GetTypeHash( Expression ), Expression, ExpressionType ) )
+				{
+					// Make the name unique
+					ChosenName = UniqueNameProvider->GenerateUniqueName( Expression.ToString() );
+					Material->GetExpression( ExpressionIndex )->SetName( *ChosenName );
+				}
+				else
+				{
+					// Keep track of the used names
+					ChosenName = Expression.ToString();
+					UniqueNameProvider->AddExistingName( ChosenName );
+				}
 
-					TFunction<void (const TSharedRef<TSet<TPair<FName, EDatasmithMaterialExpressionType>>>&, uint32, const TSharedPtr<IDatasmithUEPbrMaterialElement>&)> PushUsedNamesIntoTopLevelMaterial;
-					PushUsedNamesIntoTopLevelMaterial = [&MaterialToIndexInArray, &TopMaterialAndNamesUsedByType, &FunctionMaterialAndTheirReferencer, &PushUsedNamesIntoTopLevelMaterial] (const TSharedRef<TSet<TPair<FName, EDatasmithMaterialExpressionType>>>& UsedNames, uint32 MaterialHash, const TSharedPtr<IDatasmithUEPbrMaterialElement>& Material)
-						{
-							if ( int32* Index = MaterialToIndexInArray.FindByHash( MaterialHash, Material ) )
-							{
-								for ( const TSharedPtr<IDatasmithUEPbrMaterialElement>& Referencer : FunctionMaterialAndTheirReferencer[*Index].Value )
-								{
-									PushUsedNamesIntoTopLevelMaterial( UsedNames, GetTypeHash( Referencer ), Referencer );
-								}
-							}
-							else
-							{
-								TopMaterialAndNamesUsedByType.FindOrAddByHash( MaterialHash, Material ).Add( UsedNames );
-							}
-						};
+				NameUsed.Add( TPair<FName, EDatasmithMaterialExpressionType>( *ChosenName, ExpressionType ) );
 
+				if ( ExpressionType != EDatasmithMaterialExpressionType::Generic )
+				{
+					// Always add to the generic, we don't know what type of parameter a generic expression would be
+					TUniquePtr<FDatasmithUniqueNameProvider>& GenericUniqueNameProvider = ExpressionTypeToUniqueNameProvider.FindOrAdd( EDatasmithMaterialExpressionType::Generic );
+					if ( !GenericUniqueNameProvider )
+					{
+						GenericUniqueNameProvider = MakeUnique<FDatasmithUniqueNameProvider>();
+					}
 
-					TSet<TSharedPtr<IDatasmithUEPbrMaterialElement>> VisitedMaterials;
-					// Plus one because there is at least one top material that isn't used as a function
-					VisitedMaterials.Reserve( FunctionMaterialAndTheirReferencer.Num() + 1 );
+					GenericUniqueNameProvider->AddExistingName( ChosenName );
+					NameUsed.Add( TPair<FName, EDatasmithMaterialExpressionType>( *ChosenName, EDatasmithMaterialExpressionType::Generic ) );
+				}
+			};
 
-					/**
-					 * We expect the independents function to be at the start of the array
-					 * We want the top level function to conserve their name in case of conflicted naming of the parameters.
-					 */
-					for ( int32 Index = FunctionMaterialAndTheirReferencer.Num()-1; Index >= 0; Index-- )
+		TFunction<void (const TSharedRef<TSet<TPair<FName, EDatasmithMaterialExpressionType>>>&, uint32, const TSharedPtr<IDatasmithUEPbrMaterialElement>&)> PushUsedNamesIntoTopLevelMaterial;
+		PushUsedNamesIntoTopLevelMaterial = [&MaterialToIndexInArray, &TopMaterialAndNamesUsedByType, &FunctionMaterialAndTheirReferencer, &PushUsedNamesIntoTopLevelMaterial] (const TSharedRef<TSet<TPair<FName, EDatasmithMaterialExpressionType>>>& UsedNames, uint32 MaterialHash, const TSharedPtr<IDatasmithUEPbrMaterialElement>& Material)
+			{
+				if ( int32* Index = MaterialToIndexInArray.FindByHash( MaterialHash, Material ) )
+				{
+					for ( const TSharedPtr<IDatasmithUEPbrMaterialElement>& Referencer : FunctionMaterialAndTheirReferencer[*Index].Value )
 					{
-						const FFunctionAndMaterialsThatUseIt& FunctionAndReferencer = FunctionMaterialAndTheirReferencer[Index];
-						const TSharedPtr<IDatasmithUEPbrMaterialElement>& CurrentElement = FunctionAndReferencer.Key;
-						uint32 CurrentElementHash = GetTypeHash( CurrentElement );
-						MaterialToIndexInArray.AddByHash( CurrentElementHash, CurrentElement, Index );
+						PushUsedNamesIntoTopLevelMaterial( UsedNames, GetTypeHash( Referencer ), Referencer );
+					}
+				}
+				else
+				{
+					TopMaterialAndNamesUsedByType.FindOrAddByHash( MaterialHash, Material ).Add( UsedNames );
+				}
+			};
 
-						// Populate the name used by the Referencers
-						for ( const TSharedPtr<IDatasmithUEPbrMaterialElement>& Referencer : FunctionAndReferencer.Value )
-						{
-							uint32  ReferencerHash = GetTypeHash( Referencer );
 
-							if ( !VisitedMaterials.ContainsByHash( ReferencerHash, Referencer ) )
-							{
-								TSharedRef<TSet<TPair<FName, EDatasmithMaterialExpressionType>>> UsedNames = MakeShared<TSet<TPair<FName, EDatasmithMaterialExpressionType>>>();
-
-								FDatasmithMaterialExpressions::ForEachParamsNameInMaterial( Referencer,
-									[ReferencerHash, &Referencer, &HandleExpressionNameForMaterialFunction, &UsedNames] (FName Expression, const EDatasmithMaterialExpressionType& ExpressionType, int32 Index)
-									{
-										HandleExpressionNameForMaterialFunction( ReferencerHash, Referencer, Expression, ExpressionType, Index, UsedNames.Get() );
-									});
-								VisitedMaterials.AddByHash( ReferencerHash, Referencer );
-								PushUsedNamesIntoTopLevelMaterial( UsedNames, ReferencerHash, Referencer );
-							}
-						}
+		TSet<TSharedPtr<IDatasmithUEPbrMaterialElement>> VisitedMaterials;
+		// Plus one because there is at least one top material that isn't used as a function
+		VisitedMaterials.Reserve( FunctionMaterialAndTheirReferencer.Num() + 1 );
 
-						// For the current element
-						TSharedRef<TSet<TPair<FName, EDatasmithMaterialExpressionType>>> UsedNames = MakeShared<TSet<TPair<FName, EDatasmithMaterialExpressionType>>>();
+		/**
+			* We expect the independents function to be at the start of the array
+			* We want the top level function to conserve their name in case of conflicted naming of the parameters.
+			*/
+		for ( int32 Index = FunctionMaterialAndTheirReferencer.Num()-1; Index >= 0; Index-- )
+		{
+			const FFunctionAndMaterialsThatUseIt& FunctionAndReferencer = FunctionMaterialAndTheirReferencer[Index];
+			const TSharedPtr<IDatasmithUEPbrMaterialElement>& CurrentElement = FunctionAndReferencer.Key;
+			uint32 CurrentElementHash = GetTypeHash( CurrentElement );
+			MaterialToIndexInArray.AddByHash( CurrentElementHash, CurrentElement, Index );
 
-						TFunction<void(FName Expression, const EDatasmithMaterialExpressionType& ExpressionType, int32 Index)> ForEachParameter =
-							[CurrentElementHash, &CurrentElement, &UsedNames, &HandleExpressionNameForMaterialFunction](FName Expression, const EDatasmithMaterialExpressionType& ExpressionType, int32 Index)
-							{
-								HandleExpressionNameForMaterialFunction( CurrentElementHash, CurrentElement, Expression, ExpressionType, Index, UsedNames.Get() );
-							};
+			// Populate the name used by the Referencers
+			for ( const TSharedPtr<IDatasmithUEPbrMaterialElement>& Referencer : FunctionAndReferencer.Value )
+			{
+				uint32  ReferencerHash = GetTypeHash( Referencer );
 
-						FDatasmithMaterialExpressions::ForEachParamsNameInMaterial( CurrentElement,
-							[CurrentElementHash, &CurrentElement, &UsedNames, &HandleExpressionNameForMaterialFunction] (FName Expression, const EDatasmithMaterialExpressionType& ExpressionType, int32 Index)
-							{
-								HandleExpressionNameForMaterialFunction( CurrentElementHash, CurrentElement, Expression, ExpressionType, Index, UsedNames.Get() );
-							});
-						VisitedMaterials.AddByHash( CurrentElementHash, CurrentElement );
-						PushUsedNamesIntoTopLevelMaterial( UsedNames, CurrentElementHash, CurrentElement );
-					}
+				if ( !VisitedMaterials.ContainsByHash( ReferencerHash, Referencer ) )
+				{
+					TSharedRef<TSet<TPair<FName, EDatasmithMaterialExpressionType>>> UsedNames = MakeShared<TSet<TPair<FName, EDatasmithMaterialExpressionType>>>();
+
+					FDatasmithMaterialExpressions::ForEachParamsNameInMaterial( Referencer,
+						[ReferencerHash, &Referencer, &HandleExpressionNameForMaterialFunction, &UsedNames] (FName Expression, const EDatasmithMaterialExpressionType& ExpressionType, int32 Index)
+						{
+							HandleExpressionNameForMaterialFunction( ReferencerHash, Referencer, Expression, ExpressionType, Index, UsedNames.Get() );
+						});
+					VisitedMaterials.AddByHash( ReferencerHash, Referencer );
+					PushUsedNamesIntoTopLevelMaterial( UsedNames, ReferencerHash, Referencer );
 				}
 			}
+
+			// For the current element
+			TSharedRef<TSet<TPair<FName, EDatasmithMaterialExpressionType>>> UsedNames = MakeShared<TSet<TPair<FName, EDatasmithMaterialExpressionType>>>();
+
+			TFunction<void(FName Expression, const EDatasmithMaterialExpressionType& ExpressionType, int32 Index)> ForEachParameter =
+				[CurrentElementHash, &CurrentElement, &UsedNames, &HandleExpressionNameForMaterialFunction](FName Expression, const EDatasmithMaterialExpressionType& ExpressionType, int32 Index)
+				{
+					HandleExpressionNameForMaterialFunction( CurrentElementHash, CurrentElement, Expression, ExpressionType, Index, UsedNames.Get() );
+				};
+
+			FDatasmithMaterialExpressions::ForEachParamsNameInMaterial( CurrentElement,
+				[CurrentElementHash, &CurrentElement, &UsedNames, &HandleExpressionNameForMaterialFunction] (FName Expression, const EDatasmithMaterialExpressionType& ExpressionType, int32 Index)
+				{
+					HandleExpressionNameForMaterialFunction( CurrentElementHash, CurrentElement, Expression, ExpressionType, Index, UsedNames.Get() );
+				});
+			VisitedMaterials.AddByHash( CurrentElementHash, CurrentElement );
+			PushUsedNamesIntoTopLevelMaterial( UsedNames, CurrentElementHash, CurrentElement );
 		}
 	}
-}
+} // namespace UE::DatasmithImporter::Private
 
 void FDatasmithImporter::ImportStaticMeshes( FDatasmithImportContext& ImportContext )
 {
@@ -761,8 +752,7 @@ void FDatasmithImporter::ImportMaterials( FDatasmithImportContext& ImportContext
 			using FFunctionAndMaterialsThatUseIt = FDatasmithImporterUtils::FFunctionAndMaterialsThatUseIt;
 			TArray<FFunctionAndMaterialsThatUseIt> Functions = FDatasmithImporterUtils::GetOrderedListOfMaterialsReferencedByMaterials( ImportContext.FilteredScene );
 
-			using namespace UE::DatasmithImporter::Private::DatasmithImporter;
-			MakeChangeToExpressionsNameForMaterialFunctions(Functions);
+			UE::DatasmithImporter::Private::MakeChangeToExpressionsNameForMaterialFunctions(Functions);
 
 			for (const FFunctionAndMaterialsThatUseIt& FunctionAndMaterials : Functions)
 			{
diff --git a/Engine/Plugins/Enterprise/DatasmithImporter/Source/DatasmithImporter/Private/Utility/DatasmithImporterImpl.cpp b/Engine/Plugins/Enterprise/DatasmithImporter/Source/DatasmithImporter/Private/Utility/DatasmithImporterImpl.cpp
--- a/Engine/Plugins/Enterprise/DatasmithImporter/Source/DatasmithImporter/Private/Utility/DatasmithImporterImpl.cpp
+++ b/Engine/Plugins/Enterprise/DatasmithImporter/Source/DatasmithImporter/Private/Utility/DatasmithImporterImpl.cpp
@@ -76,11 +76,7 @@ extern UNREALED_API UEditorEngine* GEditor;
 #define LOCTEXT_NAMESPACE "DatasmithImporter"
 
 
-namespace DatasmithImporter
-{
-namespace Private
-{
-namespace DatasmithImporterImpl
+namespace UE::DatasmithImporterImpl::Private
 {
 	class FDatasmithWriter final : public FObjectWriter
 	{
@@ -113,9 +109,7 @@ namespace DatasmithImporterImpl
 			Object.Serialize(*this);
 		}
 	};
-}
-}
-}
+} // namespace UE::DatasmithImporterImpl::Private
 
 void FDatasmithImporterImpl::ReportProgress(FScopedSlowTask* SlowTask, const float ExpectedWorkThisFrame, const FText& Text)
 {
@@ -744,7 +738,7 @@ void FDatasmithImporterImpl::CopyObject(UObject& Source, UObject& Destination, T
 {
 	// Keep the previous allocation
 	TempBuffer.Reset();
-	using namespace DatasmithImporter::Private::DatasmithImporterImpl;
+	using namespace UE::DatasmithImporterImpl::Private;
 	FDatasmithWriter DatasmithWriter( Source, TempBuffer );
 	FDatasmithReader DatasmithReader( Destination, TempBuffer );
 }
diff --git a/Engine/Plugins/Enterprise/DatasmithImporter/Source/DatasmithImporter/Private/Utility/DatasmithImporterImpl.h b/Engine/Plugins/Enterprise/DatasmithImporter/Source/DatasmithImporter/Private/Utility/DatasmithImporterImpl.h
--- a/Engine/Plugins/Enterprise/DatasmithImporter/Source/DatasmithImporter/Private/Utility/DatasmithImporterImpl.h
+++ b/Engine/Plugins/Enterprise/DatasmithImporter/Source/DatasmithImporter/Private/Utility/DatasmithImporterImpl.h
@@ -2,33 +2,6 @@
 
 #pragma once
 
-//#include "DatasmithImporter.h"
-
-//#include "DatasmithActorImporter.h"
-//#include "DatasmithAdditionalData.h"
-//#include "DatasmithAssetImportData.h"
-//#include "DatasmithAssetUserData.h"
-//#include "DatasmithCameraImporter.h"
-//#include "DatasmithImportContext.h"
-//#include "DatasmithLevelSequenceImporter.h"
-//#include "DatasmithLevelVariantSetsImporter.h"
-//#include "DatasmithLightImporter.h"
-//#include "DatasmithMaterialImporter.h"
-//#include "DatasmithPayload.h"
-//#include "DatasmithPostProcessImporter.h"
-//#include "DatasmithScene.h"
-//#include "DatasmithSceneActor.h"
-//#include "DatasmithSceneFactory.h"
-//#include "DatasmithStaticMeshImporter.h"
-//#include "DatasmithTranslator.h"
-//#include "DatasmithTextureImporter.h"
-//#include "IDatasmithSceneElements.h"
-//#include "DatasmithAnimationElements.h"
-//#include "LevelVariantSets.h"
-//#include "ObjectTemplates/DatasmithObjectTemplate.h"
-//#include "Utility/DatasmithImporterUtils.h"
-//#include "Utility/DatasmithTextureResize.h"
-
 #include "CoreMinimal.h"
 
 #include "Serialization/ObjectReader.h"
