Datasmith: cleanup

diff --git a/Engine/Plugins/Enterprise/DatasmithImporter/Source/DatasmithImporter/Private/Utility/DatasmithImporterImpl.cpp b/Engine/Plugins/Enterprise/DatasmithImporter/Source/DatasmithImporter/Private/Utility/DatasmithImporterImpl.cpp
--- a/Engine/Plugins/Enterprise/DatasmithImporter/Source/DatasmithImporter/Private/Utility/DatasmithImporterImpl.cpp
+++ b/Engine/Plugins/Enterprise/DatasmithImporter/Source/DatasmithImporter/Private/Utility/DatasmithImporterImpl.cpp
@@ -4,6 +4,7 @@
 
 #include "DatasmithActorImporter.h"
 #include "DatasmithAdditionalData.h"
+#include "DatasmithAnimationElements.h"
 #include "DatasmithAssetImportData.h"
 #include "DatasmithAssetUserData.h"
 #include "DatasmithCameraImporter.h"
@@ -18,11 +19,9 @@
 #include "DatasmithSceneActor.h"
 #include "DatasmithSceneFactory.h"
 #include "DatasmithStaticMeshImporter.h"
-#include "DatasmithTranslator.h"
 #include "DatasmithTextureImporter.h"
+#include "DatasmithTranslator.h"
 #include "IDatasmithSceneElements.h"
-#include "DatasmithAnimationElements.h"
-#include "LevelVariantSets.h"
 #include "ObjectTemplates/DatasmithObjectTemplate.h"
 #include "Utility/DatasmithImporterUtils.h"
 #include "Utility/DatasmithTextureResize.h"
@@ -49,11 +48,12 @@
 #include "Landscape.h"
 #include "Layers/LayersSubsystem.h"
 #include "LevelSequence.h"
+#include "LevelVariantSets.h"
 #include "MaterialEditingLibrary.h"
-#include "MaterialShared.h"
 #include "Materials/Material.h"
 #include "Materials/MaterialFunction.h"
 #include "Materials/MaterialInstanceConstant.h"
+#include "MaterialShared.h"
 #include "Misc/FeedbackContext.h"
 #include "Misc/FileHelper.h"
 #include "Misc/PackageName.h"
@@ -64,11 +64,10 @@
 #include "Serialization/ArchiveReplaceObjectRef.h"
 #include "Serialization/MemoryReader.h"
 #include "Serialization/ObjectReader.h"
+#include "Serialization/ObjectWriter.h"
 #include "Settings/EditorExperimentalSettings.h"
 #include "SourceControlOperations.h"
 #include "Subsystems/AssetEditorSubsystem.h"
-#include "Templates/UniquePtr.h"
-#include "UObject/Package.h"
 #include "UnrealEdGlobals.h"
 
 extern UNREALED_API UEditorEngine* GEditor;
@@ -84,13 +83,13 @@ namespace UE::DatasmithImporterImpl::Private
 		FDatasmithWriter(UObject& Object, TArray< uint8 >& Bytes)
 			: FObjectWriter(Bytes)
 		{
-				SetIsLoading(false);
-				SetIsSaving(true);
-				SetIsPersistent(true);
+			SetIsLoading(false);
+			SetIsSaving(true);
+			SetIsPersistent(true);
 
-				SetPortFlags(GetPortFlags() | PPF_Duplicate);
+			SetPortFlags(GetPortFlags() | PPF_Duplicate);
 
-				Object.Serialize(*this);
+			Object.Serialize(*this);
 		}
 	};
 
@@ -119,14 +118,6 @@ void FDatasmithImporterImpl::ReportProgress(FScopedSlowTask* SlowTask, const flo
 	}
 }
 
-void FDatasmithImporterImpl::ReportProgress(FScopedSlowTask* SlowTask, const float ExpectedWorkThisFrame, FText&& Text)
-{
-	if (SlowTask)
-	{
-		SlowTask->EnterProgressFrame(ExpectedWorkThisFrame, Text);
-	}
-}
-
 bool FDatasmithImporterImpl::HasUserCancelledTask(FFeedbackContext* FeedbackContext)
 {
 	if (FeedbackContext)
@@ -402,11 +393,11 @@ void FDatasmithImporterImpl::FixReferencesForObject( UObject* Object, const TMap
 }
 
 /**
-	* Creates templates to apply the values from the SourceObject on the DestinationObject.
-	*
-	* @returns An array of template pairs. The key is the template for the object, the value is a template to force apply to the object,
-	*			it contains the values from the key and any overrides that were present on the DestinationObject.
-	*/
+ * Creates templates to apply the values from the SourceObject on the DestinationObject.
+ *
+ * @returns An array of template pairs. The key is the template for the object, the value is a template to force apply to the object,
+ *			it contains the values from the key and any overrides that were present on the DestinationObject.
+ */
 TArray< FDatasmithImporterImpl::FMigratedTemplatePairType > FDatasmithImporterImpl::MigrateTemplates( UObject* SourceObject, UObject* DestinationObject, const TMap< UObject*, UObject* >* ReferencesToRemap, bool bIsForActor )
 {
 	TArray< FMigratedTemplatePairType > Results;
@@ -455,13 +446,13 @@ TArray< FDatasmithImporterImpl::FMigratedTemplatePairType > FDatasmithImporterIm
 }
 
 /**
-	* Applies the templates created from MigrateTemplates to DestinationObject.
-	*
-	* For an Object A that should be duplicated over an existing A', for which we want to keep the Datasmith overrides:
-	* - Call MigrateTemplates(A, A')
-	* - Duplicate A over A'
-	* - ApplyMigratedTemplates(A')
-	*/
+ * Applies the templates created from MigrateTemplates to DestinationObject.
+ *
+ * For an Object A that should be duplicated over an existing A', for which we want to keep the Datasmith overrides:
+ * - Call MigrateTemplates(A, A')
+ * - Duplicate A over A'
+ * - ApplyMigratedTemplates(A')
+ */
 void FDatasmithImporterImpl::ApplyMigratedTemplates( TArray< FMigratedTemplatePairType >& MigratedTemplates, UObject* DestinationObject )
 {
 	for ( FMigratedTemplatePairType& MigratedTemplate : MigratedTemplates )
diff --git a/Engine/Plugins/Enterprise/DatasmithImporter/Source/DatasmithImporter/Private/Utility/DatasmithImporterImpl.h b/Engine/Plugins/Enterprise/DatasmithImporter/Source/DatasmithImporter/Private/Utility/DatasmithImporterImpl.h
--- a/Engine/Plugins/Enterprise/DatasmithImporter/Source/DatasmithImporter/Private/Utility/DatasmithImporterImpl.h
+++ b/Engine/Plugins/Enterprise/DatasmithImporter/Source/DatasmithImporter/Private/Utility/DatasmithImporterImpl.h
@@ -4,16 +4,14 @@
 
 #include "CoreMinimal.h"
 
-#include "Serialization/ObjectReader.h"
-#include "Serialization/ObjectWriter.h"
 #include "UObject/StrongObjectPtr.h"
 
 class AActor;
 class FFeedbackContext;
+class FMaterialUpdateContext;
 class IDatasmithBaseMaterialElement;
 class UActorComponent;
 class UDatasmithObjectTemplate;
-class FMaterialUpdateContext;
 class UMaterial;
 class UMaterialInterface;
 class USceneComponent;
@@ -28,8 +26,6 @@ class FDatasmithImporterImpl
 public:
 	static void ReportProgress(FScopedSlowTask* SlowTask, const float ExpectedWorkThisFrame, const FText& Text);
 
-	static void ReportProgress(FScopedSlowTask* SlowTask, const float ExpectedWorkThisFrame, FText&& Text);
-
 	static bool HasUserCancelledTask(FFeedbackContext* FeedbackContext);
 
 	static UObject* PublicizeAsset( UObject* SourceAsset, const TCHAR* DestinationPath, UObject* ExistingAsset );
