Datasmith: cloth properties basic support: translation from CLO and import as assets

diff --git a/Engine/Plugins/Enterprise/DatasmithImporter/Source/DatasmithImporter/Private/DatasmithImporter.cpp b/Engine/Plugins/Enterprise/DatasmithImporter/Source/DatasmithImporter/Private/DatasmithImporter.cpp
--- a/Engine/Plugins/Enterprise/DatasmithImporter/Source/DatasmithImporter/Private/DatasmithImporter.cpp
+++ b/Engine/Plugins/Enterprise/DatasmithImporter/Source/DatasmithImporter/Private/DatasmithImporter.cpp
@@ -370,9 +370,11 @@ void FDatasmithImporter::ImportStaticMeshes( FDatasmithImportContext& ImportCont
 
 void FDatasmithImporter::ImportClothes(FDatasmithImportContext& ImportContext)
 {
+	TRACE_CPUPROFILER_EVENT_SCOPE(FDatasmithImporter::ImportClothes);
+
 	// #ue_ds_cloth_note: FDatasmithImporter::ImportClothes: experimental import code
 	const int32 ClothesCount = ImportContext.FilteredScene->GetClothesCount();
-	if ( !ImportContext.Options->BaseOptions.bIncludeGeometry || ClothesCount == 0 )
+	if ( !ImportContext.Options->BaseOptions.bIncludeGeometry || ClothesCount == 0 || ImportContext.SceneTranslator == nullptr)
 	{
 		return;
 	}
@@ -390,13 +392,9 @@ void FDatasmithImporter::ImportClothes(FDatasmithImportContext& ImportContext)
 
 		// #ue_ds_cloth_todo: check existing asset
 
-		{ // FDatasmithImporter::ImportStaticMesh
-
-			{ // FDatasmithStaticMeshImporter::ImportStaticMesh
-//
-				// #ue_ds_cloth_todo: import UChaosClothPreset here
-// 				ClothElement->GetPhyDescriptions.....
+		{ // copy of FDatasmithImporter::ImportStaticMesh
 
+			{ // copy of FDatasmithStaticMeshImporter::ImportStaticMesh
 // 				// 2. get destination asset
 // 				// 2.1. FDatasmithImporterUtils::CanCreateAsset<UStaticMesh>
 				int32 MaxNameCharCount = FDatasmithImporterUtils::GetAssetNameMaxCharCount(ImportContext.AssetsContext.StaticMeshesFinalPackage.Get());
@@ -404,25 +402,22 @@ void FDatasmithImporter::ImportClothes(FDatasmithImportContext& ImportContext)
 //
 // 				// if !existing
 				UPackage* Outer = ImportContext.AssetsContext.StaticMeshesImportPackage.Get();
-				EObjectFlags ObjectFlags = ImportContext.ObjectFlags & ~RF_Public; // not RF_Public, the publicize asset will be public
-
-				FDatasmithCloth DsCloth;
+				EObjectFlags ObjectFlags = ImportContext.ObjectFlags & ~RF_Public; // not RF_Public yet, the publicized asset will be.
 
-				if (ImportContext.SceneTranslator)
+				FDatasmithClothElementPayload Payload;
+				// #ue_ds_cloth_todo: async parallel support
 				{
-					FDatasmithClothElementPayload Payload;
-					if (ImportContext.SceneTranslator->LoadCloth(ClothElement, Payload))
-					{
-						DsCloth = MoveTemp(Payload.Cloth);
-					}
-					else
+					TRACE_CPUPROFILER_EVENT_SCOPE(FDatasmithImporter::ImportClothes:TranslatorLoadCloth);
+					if (!ImportContext.SceneTranslator->LoadCloth(ClothElement, Payload))
 					{
 						UE_LOG(LogDatasmithImport, Warning, TEXT("Cloth %s: load failed, translator issue."), ClothElement->GetName());
 						return;
 					}
 				}
+				FDatasmithCloth DsCloth = MoveTemp(Payload.Cloth);
 
 				UObject* ClothAssetObj = nullptr;
+				TArray<UObject*> PropertyAssetsObjects;
 
 #if 0 // #ue_ds_cloth_arch switch between local impl or delegates to IDatasmithImporterExt
 				UChaosClothAsset* ClothAsset = NewObject<UChaosClothAsset>(Outer, *ClothName, ObjectFlags);
@@ -432,12 +427,25 @@ void FDatasmithImporter::ImportClothes(FDatasmithImportContext& ImportContext)
 #else
 				if (IDatasmithImporterExt* Ext = IDatasmithImporterModule::Get().GetClothImporterExtension())
 				{
+					for (const FDatasmithClothPresetPropertySet& PS : DsCloth.PropertySets)
+					{
+						if (UObject* PropertySetObject = Ext->MakeClothPropertyAsset(Outer, *PS.SetName, ObjectFlags))
+						{
+							PropertyAssetsObjects.Add(PropertySetObject);
+							Ext->FillPropertySet(PropertySetObject, ClothElement, PS);
+						}
+					}
+
 					ClothAssetObj = Ext->MakeClothAsset(Outer, *ClothName, ObjectFlags);
 					Ext->FillCloth(ClothAssetObj, ClothElement, DsCloth);
 				}
 #endif
 
 				ImportContext.ImportedClothes.Add(ClothElement, ClothAssetObj);
+				for (UObject* PropertyAsset : PropertyAssetsObjects)
+				{
+					ImportContext.ImportedClothPresets.Add(PropertyAsset);
+				}
 			}
 
 			// #ue_ds_cloth_todo: CreateStaticMeshAssetImportData
@@ -1789,6 +1797,7 @@ void FDatasmithImporter::FinalizeImport(FDatasmithImportContext& ImportContext,
 
 	const int32 NumImportedObjects = ImportContext.ImportedStaticMeshes.Num() +
 									 ImportContext.ImportedClothes.Num() +
+									 ImportContext.ImportedClothPresets.Num() +
 									 ImportContext.ImportedTextures.Num() +
 									 ImportContext.ImportedMaterialFunctions.Num() +
 									 ImportContext.ImportedMaterials.Num() +
@@ -1991,6 +2000,27 @@ void FDatasmithImporter::FinalizeImport(FDatasmithImportContext& ImportContext,
 	FDatasmithStaticMeshImporter::BuildStaticMeshes(StaticMeshes.Array(), ProgressFunction);
 
 	// #ue_ds_cloth_note: experimental FinalizeImport code
+	for (auto& SourcePresetObj : ImportContext.ImportedClothPresets)
+	{
+		if (!SourcePresetObj || (ValidAssets.Num() > 0 && !ValidAssets.Contains(SourcePresetObj)))
+		{
+			continue;
+		}
+
+		FDatasmithImporterImpl::ReportProgress(Progress, 1.f, FText::FromString(FString::Printf(TEXT("Finalizing assets %d/%d (Cloth Preset %s) ..."), ++AssetIndex, NumAssetsToFinalize, *SourcePresetObj->GetName())));
+
+		UObject* ExistingCloth = nullptr; // #ue_ds_cloth_todo but requires a dependency DatasmithContent -> ChaosClothAsset.uplugin
+
+		FString SourcePackagePath = SourcePresetObj->GetOutermost()->GetName();
+		FString DestinationPackagePath = SourcePackagePath.Replace( *TransientFolderPath, *RootFolderPath, ESearchCase::CaseSensitive );
+
+		UObject* FinalizedPreset = FDatasmithImporterImpl::FinalizeAsset(SourcePresetObj, *DestinationPackagePath, ExistingCloth, &ReferencesToRemap);
+
+		FDatasmithImporterImpl::CheckAssetPersistenceValidity(FinalizedPreset, ImportContext);
+
+		SourcePresetObj = FinalizedPreset;
+	}
+
 	for (auto& ImportedPair : ImportContext.ImportedClothes)
 	{
 		UObject* SourceClothObj = ImportedPair.Value;
diff --git a/Engine/Plugins/Enterprise/DatasmithImporter/Source/DatasmithImporter/Public/DatasmithClothImporter.h b/Engine/Plugins/Enterprise/DatasmithImporter/Source/DatasmithImporter/Public/DatasmithClothImporter.h
--- a/Engine/Plugins/Enterprise/DatasmithImporter/Source/DatasmithImporter/Public/DatasmithClothImporter.h
+++ b/Engine/Plugins/Enterprise/DatasmithImporter/Source/DatasmithImporter/Public/DatasmithClothImporter.h
@@ -7,6 +7,7 @@
 #include "UObject/ObjectMacros.h"
 
 class FDatasmithCloth;
+class FDatasmithClothPresetPropertySet;
 class IDatasmithClothElement;
 
 
@@ -15,6 +16,9 @@ class IDatasmithImporterExt
 {
 public:
 	virtual UObject* MakeClothAsset(UObject* Outer, const TCHAR* Name, EObjectFlags ObjectFlags) = 0;
-	virtual void FillCloth(UObject* ClothAsset, TSharedRef<IDatasmithClothElement> ClothElement, FDatasmithCloth& DsCloth) = 0;
+	virtual void FillCloth(UObject* ClothAsset, TSharedRef<IDatasmithClothElement> ClothElement, FDatasmithCloth& Cloth) = 0;
+
+	virtual UObject* MakeClothPropertyAsset(UObject* Outer, const TCHAR* Name, EObjectFlags ObjectFlags) = 0;
+	virtual void FillPropertySet(UObject* PropertySetAsset, TSharedRef<IDatasmithClothElement> ClothElement, const FDatasmithClothPresetPropertySet& PropertySet) = 0;
 };
 
diff --git a/Engine/Plugins/Enterprise/DatasmithImporter/Source/DatasmithImporter/Public/DatasmithImportContext.h b/Engine/Plugins/Enterprise/DatasmithImporter/Source/DatasmithImporter/Public/DatasmithImportContext.h
--- a/Engine/Plugins/Enterprise/DatasmithImporter/Source/DatasmithImporter/Public/DatasmithImportContext.h
+++ b/Engine/Plugins/Enterprise/DatasmithImporter/Source/DatasmithImporter/Public/DatasmithImportContext.h
@@ -214,7 +214,7 @@ struct DATASMITHIMPORTER_API FDatasmithImportContext
 	FDatasmithImportContext(const FString& InFileName, bool bLoadConfig, const FName& LoggerName, const FText& LoggerLabel, TSharedPtr<IDatasmithTranslator> InSceneTranslator = nullptr);
 
 	FDatasmithImportContext(const TSharedPtr<UE::DatasmithImporter::FExternalSource>& InExternalSource, bool bLoadConfig, const FName& LoggerName, const FText& LoggerLabel);
-	
+
 	/** Cached MD5 hash value for faster processing */
 	FMD5Hash FileHash;
 
@@ -257,7 +257,9 @@ struct DATASMITHIMPORTER_API FDatasmithImportContext
 	/** Map of imported mesh for each mesh element */
 	TMap< TSharedRef< IDatasmithMeshElement >, UStaticMesh* > ImportedStaticMeshes;
 
-	TMap< TSharedRef< IDatasmithClothElement >, UObject* > ImportedClothes; // UChaosClothAsset
+	TMap< TSharedRef< IDatasmithClothElement >, UObject* > ImportedClothes; // #ue_ds_cloth_todo: UChaosClothAsset
+
+	TArray< UObject* > ImportedClothPresets; // #ue_ds_cloth_todo: UChaosClothPreset // #ue_ds_cloth_todo: map with a dedicated element so that clothes can share parameters
 
 	/** Register IDatasmithMeshElement by their name so they can be searched faster */
 	TMap< FString, TSharedRef < IDatasmithMeshElement > > ImportedStaticMeshesByName;
@@ -316,7 +318,7 @@ public:
 	 * @param bSilent				Doesn't display the options dialog and skips other user input requests
 	 */
 	bool Init(const FString& InImportPath, EObjectFlags InFlags, FFeedbackContext* InWarn, const TSharedPtr<FJsonObject>& ImportSettingsJson, bool bSilent);
-	
+
 	/**
 	 * First part of the Init process, replaces Init, requires a call to SetupDestination and InitScene after that.
 	 * Displays the options to the end-user (blocking call), and updates translator accordingly.
diff --git a/Engine/Source/Runtime/Datasmith/DatasmithCore/Private/DatasmithMeshSerialization.cpp b/Engine/Source/Runtime/Datasmith/DatasmithCore/Private/DatasmithMeshSerialization.cpp
--- a/Engine/Source/Runtime/Datasmith/DatasmithCore/Private/DatasmithMeshSerialization.cpp
+++ b/Engine/Source/Runtime/Datasmith/DatasmithCore/Private/DatasmithMeshSerialization.cpp
@@ -319,7 +319,7 @@ FDatasmithPackedMeshes GetDatasmithMeshFromFile(const FString& MeshPath)
 
 DATASMITHCORE_API FDatasmithPackedMeshes GetDatasmithClothFromFile(const FString& Path)
 {
-	// #ue_ds_cloth: asset serialization
+	// #ue_ds_cloth_note: asset serialization
 	FDatasmithPackedMeshes Result;
 
 	TUniquePtr<FArchive> Archive( IFileManager::Get().CreateFileReader(*Path) );
diff --git a/Engine/Source/Runtime/Datasmith/DatasmithCore/Private/DatasmithSceneElementsImpl.h b/Engine/Source/Runtime/Datasmith/DatasmithCore/Private/DatasmithSceneElementsImpl.h
--- a/Engine/Source/Runtime/Datasmith/DatasmithCore/Private/DatasmithSceneElementsImpl.h
+++ b/Engine/Source/Runtime/Datasmith/DatasmithCore/Private/DatasmithSceneElementsImpl.h
@@ -387,7 +387,7 @@ private:
 	TDatasmithReferenceArrayProxy<IDatasmithMaterialIDElement> MaterialSlots;
 };
 
-// #ue_ds_cloth doc: Experimental Element that describes a cloth asset
+// #ue_ds_cloth_doc: Experimental Element that describes a cloth asset
 class FDatasmithClothElementImpl : public FDatasmithElementImpl< IDatasmithClothElement >
 {
 public:
diff --git a/Engine/Source/Runtime/Datasmith/DatasmithCore/Private/DatasmithSceneXmlReader.cpp b/Engine/Source/Runtime/Datasmith/DatasmithCore/Private/DatasmithSceneXmlReader.cpp
--- a/Engine/Source/Runtime/Datasmith/DatasmithCore/Private/DatasmithSceneXmlReader.cpp
+++ b/Engine/Source/Runtime/Datasmith/DatasmithCore/Private/DatasmithSceneXmlReader.cpp
@@ -437,7 +437,7 @@ void FDatasmithSceneXmlReader::ParseMesh(FXmlNode* InNode, TSharedPtr<IDatasmith
 
 void FDatasmithSceneXmlReader::ParseCloth(FXmlNode* InNode, TSharedPtr<IDatasmithClothElement>& OutElement) const
 {
-	// #ue_ds_cloth: xml read
+	// #ue_ds_cloth_note: xml read
 	ParseElement( InNode, OutElement.ToSharedRef() );
 
 	for (const FXmlNode* Node : InNode->GetChildrenNodes())
@@ -1215,7 +1215,7 @@ bool FDatasmithSceneXmlReader::ParseXmlFile(TSharedRef< IDatasmithScene >& OutSc
 
 			OutScene->AddCloth(Element);
 
-// 			Objects.Add( Element->GetName(), Element ); // #ue_ds_cloth_todo cloths referencables by other elements
+// 			Objects.Add( Element->GetName(), Element ); // #ue_ds_cloth_todo cloths referencable by other elements
 		}
 		// LEVEL SEQUENCES
 		else if (Nodes[i]->GetTag() == DATASMITH_LEVELSEQUENCENAME)
diff --git a/Engine/Source/Runtime/Datasmith/DatasmithCore/Public/DatasmithCloth.h b/Engine/Source/Runtime/Datasmith/DatasmithCore/Public/DatasmithCloth.h
--- a/Engine/Source/Runtime/Datasmith/DatasmithCore/Public/DatasmithCloth.h
+++ b/Engine/Source/Runtime/Datasmith/DatasmithCore/Public/DatasmithCloth.h
@@ -1,32 +1,48 @@
-// Copyright Epic Games, Inc. All Rights Reserved.
-#pragma once
-
-#include "CoreMinimal.h"
-
-class FDatasmithMesh;
-
-class DATASMITHCORE_API FDatasmithClothPattern
-{
-public:
-	TArray<FVector2f> SimPosition;
-	TArray<FVector3f> SimRestPosition;
-	TArray<uint32> SimTriangleIndices;
-
-public:
-	bool IsValid() { return SimRestPosition.Num() == SimPosition.Num() && SimTriangleIndices.Num() % 3 == 0; }
-};
-
-/*
- * Structure of a Cloth
- * - One Thin Mesh:
- * - base for physics simulation
- *
- * - One Thick Mesh
- * - aka. render mesh
- */
-class DATASMITHCORE_API FDatasmithCloth
-{
-public:
-	TArray<FDatasmithClothPattern> Patterns;
-};
-
+// Copyright Epic Games, Inc. All Rights Reserved.
+#pragma once
+
+#include "CoreMinimal.h"
+
+class FDatasmithMesh;
+
+class DATASMITHCORE_API FDatasmithClothPattern
+{
+public:
+	TArray<FVector2f> SimPosition;
+	TArray<FVector3f> SimRestPosition;
+	TArray<uint32> SimTriangleIndices;
+
+public:
+	bool IsValid() { return SimRestPosition.Num() == SimPosition.Num() && SimTriangleIndices.Num() % 3 == 0; }
+};
+
+
+class FDatasmithClothPresetProperty
+{
+public:
+	FName Name;
+	double Value;
+};
+
+class DATASMITHCORE_API FDatasmithClothPresetPropertySet
+{
+public:
+	FString SetName;
+	TArray<FDatasmithClothPresetProperty> Properties;
+};
+
+/*
+ * Structure of a Cloth
+ * - One Thin Mesh:
+ * - base for physics simulation
+ *
+ * - One Thick Mesh
+ * - aka. render mesh
+ */
+class DATASMITHCORE_API FDatasmithCloth
+{
+public:
+	TArray<FDatasmithClothPattern> Patterns;
+	TArray<FDatasmithClothPresetPropertySet> PropertySets;
+};
+
