Datasmith: Support cloth reimport

diff --git a/Engine/Plugins/Enterprise/DatasmithContent/Source/DatasmithContent/Public/DatasmithScene.h b/Engine/Plugins/Enterprise/DatasmithContent/Source/DatasmithContent/Public/DatasmithScene.h
--- a/Engine/Plugins/Enterprise/DatasmithContent/Source/DatasmithContent/Public/DatasmithScene.h
+++ b/Engine/Plugins/Enterprise/DatasmithContent/Source/DatasmithContent/Public/DatasmithScene.h
@@ -41,6 +41,10 @@ public:
 	UPROPERTY(VisibleAnywhere, Category="Datasmith", AdvancedDisplay)
 	TMap< FName, TSoftObjectPtr< UStaticMesh > > StaticMeshes;
 
+	/** Map of all the cloth related to this Datasmith Scene */
+	UPROPERTY(VisibleAnywhere, Category="Datasmith", AdvancedDisplay)
+	TMap< FName, TSoftObjectPtr< UObject > > Clothes; // UChaosClothAsset
+
 	/** Map of all the textures related to this Datasmith Scene */
 	UPROPERTY(VisibleAnywhere, Category="Datasmith", AdvancedDisplay)
 	TMap< FName, TSoftObjectPtr< UTexture > > Textures;
diff --git a/Engine/Plugins/Enterprise/DatasmithImporter/Source/DatasmithImporter/Private/DatasmithImporter.cpp b/Engine/Plugins/Enterprise/DatasmithImporter/Source/DatasmithImporter/Private/DatasmithImporter.cpp
--- a/Engine/Plugins/Enterprise/DatasmithImporter/Source/DatasmithImporter/Private/DatasmithImporter.cpp
+++ b/Engine/Plugins/Enterprise/DatasmithImporter/Source/DatasmithImporter/Private/DatasmithImporter.cpp
@@ -370,15 +370,29 @@ void FDatasmithImporter::ImportStaticMeshes( FDatasmithImportContext& ImportCont
 
 void FDatasmithImporter::ImportClothes(FDatasmithImportContext& ImportContext)
 {
-	TRACE_CPUPROFILER_EVENT_SCOPE(FDatasmithImporter::ImportClothes);
-
 	// #ue_ds_cloth_note: FDatasmithImporter::ImportClothes: experimental import code
 	const int32 ClothesCount = ImportContext.FilteredScene->GetClothesCount();
-	if ( !ImportContext.Options->BaseOptions.bIncludeGeometry || ClothesCount == 0 || ImportContext.SceneTranslator == nullptr)
+
+	if (!ImportContext.Options->BaseOptions.bIncludeGeometry
+	 || ClothesCount == 0
+	 || ImportContext.SceneTranslator == nullptr
+	 || !ImportContext.AssetsContext.StaticMeshesFinalPackage
+	 || ImportContext.AssetsContext.StaticMeshesFinalPackage->GetFName() == NAME_None)
 	{
 		return;
 	}
 
+	TRACE_CPUPROFILER_EVENT_SCOPE(FDatasmithImporter::ImportClothes);
+
+	TUniquePtr<FScopedSlowTask> ProgressPtr;
+	if (ImportContext.FeedbackContext)
+	{
+		ProgressPtr = MakeUnique<FScopedSlowTask>(ClothesCount, LOCTEXT("ImportClothes", "Importing Clothes..."), true, *ImportContext.FeedbackContext);
+		ProgressPtr->MakeDialog(true);
+	}
+
+	FScopedSlowTask* Progress = ProgressPtr.Get();
+
 	for (int32 ClothIndex = 0; ClothIndex < ClothesCount && !ImportContext.bUserCancelled; ++ClothIndex)
 	{
 		ImportContext.bUserCancelled = FDatasmithImporterImpl::HasUserCancelledTask(ImportContext.FeedbackContext);
@@ -390,7 +404,16 @@ void FDatasmithImporter::ImportClothes(FDatasmithImportContext& ImportContext)
 		}
 		TSharedRef<IDatasmithClothElement> ClothElement = ClothElementPtr.ToSharedRef();
 
-		// #ue_ds_cloth_todo: check existing asset
+		FDatasmithImporterImpl::ReportProgress(Progress, 1.f, FText::FromString(FString::Printf(TEXT("Importing cloth %d/%d (%s) ..."), ClothIndex + 1, ClothesCount, ClothElement->GetLabel())));
+
+		UObject* ExistingCloth = nullptr;
+		if (ImportContext.SceneAsset)
+		{
+			if (TSoftObjectPtr<UObject>* ExistingAssetPtr = ImportContext.SceneAsset->Clothes.Find(ClothElement->GetName()))
+			{
+				ExistingCloth = ExistingAssetPtr->LoadSynchronous();
+			}
+		}
 
 		{ // copy of FDatasmithImporter::ImportStaticMesh
 
@@ -399,8 +422,8 @@ void FDatasmithImporter::ImportClothes(FDatasmithImportContext& ImportContext)
 // 				// 2.1. FDatasmithImporterUtils::CanCreateAsset<UStaticMesh>
 				int32 MaxNameCharCount = FDatasmithImporterUtils::GetAssetNameMaxCharCount(ImportContext.AssetsContext.StaticMeshesFinalPackage.Get());
 				FString ClothName = ImportContext.AssetsContext.StaticMeshNameProvider.GenerateUniqueName(ClothElement->GetLabel(), MaxNameCharCount);
-//
-// 				// if !existing
+				ClothName = ObjectTools::SanitizeObjectName(ClothName);
+
 				UPackage* Outer = ImportContext.AssetsContext.StaticMeshesImportPackage.Get();
 				EObjectFlags ObjectFlags = ImportContext.ObjectFlags & ~RF_Public; // not RF_Public yet, the publicized asset will be.
 
@@ -416,41 +439,64 @@ void FDatasmithImporter::ImportClothes(FDatasmithImportContext& ImportContext)
 				}
 				FDatasmithCloth DsCloth = MoveTemp(Payload.Cloth);
 
-				UObject* ClothAssetObj = nullptr;
-				TArray<UObject*> PropertyAssetsObjects;
+				UObject* ClothAsset = nullptr;
+				TArray<UObject*> ClothPresetAssets;
 
-#if 0 // #ue_ds_cloth_arch switch between local impl or delegates to IDatasmithImporterExt
-				UChaosClothAsset* ClothAsset = NewObject<UChaosClothAsset>(Outer, *ClothName, ObjectFlags);
-				// #ue_ds_cloth_todo: markPackageDirty?
-				UE::Datasmith::ClothImporter::FillCloth(ClothAsset, ClothElement, DsCloth);
-				ClothAssetObj = (UObject*)ClothAsset;
-#else
+				// #ue_ds_cloth_arch The ImporterExtension code is expected to be temporary
 				if (IDatasmithImporterExt* Ext = IDatasmithImporterModule::Get().GetClothImporterExtension())
 				{
 					for (const FDatasmithClothPresetPropertySet& PS : DsCloth.PropertySets)
 					{
-						if (UObject* PropertySetObject = Ext->MakeClothPropertyAsset(Outer, *PS.SetName, ObjectFlags))
+						FString PresetAssetName = ObjectTools::SanitizeObjectName(PS.SetName);
+						if (UObject* PropertySetObject = Ext->MakeClothPropertyAsset(Outer, *PresetAssetName, ObjectFlags))
 						{
-							PropertyAssetsObjects.Add(PropertySetObject);
+							ClothPresetAssets.Add(PropertySetObject);
 							Ext->FillPropertySet(PropertySetObject, ClothElement, PS);
 						}
 					}
 
-					ClothAssetObj = Ext->MakeClothAsset(Outer, *ClothName, ObjectFlags);
-					Ext->FillCloth(ClothAssetObj, ClothElement, DsCloth);
+					if (ExistingCloth)
+					{
+						if (ExistingCloth->GetOuter() != Outer)
+						{
+							// #ue_ds_todo why do we duplicate assets here ?
+							// We create a new asset from an existing instead of a new one, but we strip it from template information... Why ?
+
+							// Temporary flag to skip PostLoad during DuplicateObject
+							ExistingCloth->SetFlags(RF_ArchetypeObject);
+							ClothAsset = ::DuplicateObject<UObject>(ExistingCloth, Outer, *ClothName);
+							ExistingCloth->ClearFlags(RF_ArchetypeObject);
+							ClothAsset->ClearFlags(RF_ArchetypeObject);
+
+							IDatasmithImporterModule::Get().ResetOverrides(ClothAsset);
+						}
+						else
+						{
+							ClothAsset = ExistingCloth;
+						}
+
+						ExistingCloth->SetFlags(ObjectFlags);
+					}
+					else
+					{
+						ClothAsset = Ext->MakeClothAsset(Outer, *ClothName, ObjectFlags);
+					}
+
+					Ext->FillCloth(ClothAsset, ClothElement, DsCloth);
 				}
-#endif
 
-				ImportContext.ImportedClothes.Add(ClothElement, ClothAssetObj);
-				for (UObject* PropertyAsset : PropertyAssetsObjects)
+				ImportContext.ImportedClothes.Add(ClothElement, ClothAsset);
+				for (UObject* ClothPresetAsset : ClothPresetAssets)
 				{
-					ImportContext.ImportedClothPresets.Add(PropertyAsset);
+					ImportContext.ImportedClothPresets.Add(ClothPresetAsset);
 				}
 			}
 
 			// #ue_ds_cloth_todo: CreateStaticMeshAssetImportData
 			// #ue_ds_cloth_todo: ImportMetaDataForObject
 		}
+
+		// #ue_ds_cloth_todo: ImportContext.ImportedStaticMeshesByName.Add(MeshElement->GetName(), MeshElement); // see FDatasmithFindAssetTypeHelper
 	}
 }
 
@@ -1221,7 +1267,7 @@ void FDatasmithImporter::FinalizeActors( FDatasmithImportContext& ImportContext,
 				{
 					if ( UAssetUserData* AssetUserData = Cast<UAssetUserData>( SubObject ) )
 					{
-						AssetUserData->SetFlags( AssetUserData->GetFlags() | RF_Public );
+						AssetUserData->SetFlags( RF_Public );
 					}
 				}
 
@@ -2009,12 +2055,13 @@ void FDatasmithImporter::FinalizeImport(FDatasmithImportContext& ImportContext,
 
 		FDatasmithImporterImpl::ReportProgress(Progress, 1.f, FText::FromString(FString::Printf(TEXT("Finalizing assets %d/%d (Cloth Preset %s) ..."), ++AssetIndex, NumAssetsToFinalize, *SourcePresetObj->GetName())));
 
-		UObject* ExistingCloth = nullptr; // #ue_ds_cloth_todo but requires a dependency DatasmithContent -> ChaosClothAsset.uplugin
+		// #ue_ds_cloth_existing presets
+		UObject* ExistingPreset = nullptr; // Cloth = ImportContext.SceneAsset ? ImportContext.SceneAsset->Clothes.FindOrAdd(StaticMeshId).Get() : nullptr;
 
 		FString SourcePackagePath = SourcePresetObj->GetOutermost()->GetName();
 		FString DestinationPackagePath = SourcePackagePath.Replace( *TransientFolderPath, *RootFolderPath, ESearchCase::CaseSensitive );
 
-		UObject* FinalizedPreset = FDatasmithImporterImpl::FinalizeAsset(SourcePresetObj, *DestinationPackagePath, ExistingCloth, &ReferencesToRemap);
+		UObject* FinalizedPreset = FDatasmithImporterImpl::FinalizeAsset(SourcePresetObj, *DestinationPackagePath, ExistingPreset, &ReferencesToRemap);
 
 		FDatasmithImporterImpl::CheckAssetPersistenceValidity(FinalizedPreset, ImportContext);
 
@@ -2033,17 +2080,16 @@ void FDatasmithImporter::FinalizeImport(FDatasmithImportContext& ImportContext,
 		FDatasmithImporterImpl::ReportProgress(Progress, 1.f, FText::FromString(FString::Printf(TEXT("Finalizing assets %d/%d (Cloth %s) ..."), ++AssetIndex, NumAssetsToFinalize, *SourceClothObj->GetName())));
 
 		FName ClothId = ImportedPair.Key->GetName();
-		UObject* ExistingCloth = nullptr; // #ue_ds_cloth_todo but requires a dependency DatasmithContent -> ChaosClothAsset.uplugin
-// 		UStaticMesh* ExistingCloth = ImportContext.SceneAsset ? ImportContext.SceneAsset->Clothes.FindOrAdd(ClothId).Get() : nullptr;
+		UObject* ExistingCloth = ImportContext.SceneAsset ? ImportContext.SceneAsset->Clothes.FindOrAdd(ClothId).Get() : nullptr;
 
 		FString SourcePackagePath = SourceClothObj->GetOutermost()->GetName();
 		FString DestinationPackagePath = SourcePackagePath.Replace( *TransientFolderPath, *RootFolderPath, ESearchCase::CaseSensitive );
 
 		UObject* FinalizedCloth = FinalizeCloth(SourceClothObj, *DestinationPackagePath, ExistingCloth, &ReferencesToRemap);
-// 		if (ImportContext.SceneAsset)
-// 		{
-// 			ImportContext.SceneAsset->Clothes[ClothId] = FinalizedCloth; // #ue_ds_cloth_todo same constraint
-// 		}
+		if (ImportContext.SceneAsset)
+		{
+			ImportContext.SceneAsset->Clothes[ClothId] = FinalizedCloth;
+		}
 		FDatasmithImporterImpl::CheckAssetPersistenceValidity(FinalizedCloth, ImportContext);
 
 		ImportedPair.Value = FinalizedCloth;
