Datasmith: Fix inconsistent handling of invalid char between the publicized asset, and the search code of an existing asset.

diff --git a/Engine/Plugins/Enterprise/DatasmithImporter/Source/DatasmithImporter/Private/Utility/DatasmithImporterImpl.cpp b/Engine/Plugins/Enterprise/DatasmithImporter/Source/DatasmithImporter/Private/Utility/DatasmithImporterImpl.cpp
--- a/Engine/Plugins/Enterprise/DatasmithImporter/Source/DatasmithImporter/Private/Utility/DatasmithImporterImpl.cpp
+++ b/Engine/Plugins/Enterprise/DatasmithImporter/Source/DatasmithImporter/Private/Utility/DatasmithImporterImpl.cpp
@@ -132,15 +132,16 @@ UObject* FDatasmithImporterImpl::PublicizeAsset( UObject* SourceAsset, const TCH
 {
 	TRACE_CPUPROFILER_EVENT_SCOPE(FDatasmithImporterImpl::PublicizeAsset);
 
-	UPackage* DestinationPackage;
+	UPackage* DestinationPackage = nullptr;
+
+	const FString AssetName = UPackageTools::SanitizePackageName(SourceAsset->GetName());
 
 	if ( !ExistingAsset )
 	{
-		const FString AssetName = SourceAsset->GetName();
-		bool bPathIsComplete = AssetName == FPaths::GetBaseFilename( DestinationPath );
-		FString DestinationPackagePath = UPackageTools::SanitizePackageName( bPathIsComplete ? DestinationPath : FPaths::Combine( DestinationPath, AssetName ) );
-		FString DestinationAssetPath = DestinationPackagePath + TEXT(".") + UPackageTools::SanitizePackageName( AssetName );
-
+		FString SanitizedDestinationPath = UPackageTools::SanitizePackageName( DestinationPath );
+		bool bPathIsComplete = AssetName == FPaths::GetBaseFilename( SanitizedDestinationPath );
+		FString DestinationPackagePath = bPathIsComplete ? SanitizedDestinationPath : FPaths::Combine( SanitizedDestinationPath, AssetName );
+		FString DestinationAssetPath = DestinationPackagePath + TEXT(".") + AssetName;
 		ExistingAsset = FDatasmithImporterUtils::FindObject<UObject>( nullptr, DestinationAssetPath );
 
 		DestinationPackage = ExistingAsset ? ExistingAsset->GetOutermost() : CreatePackage( *DestinationPackagePath );
@@ -173,7 +174,7 @@ UObject* FDatasmithImporterImpl::PublicizeAsset( UObject* SourceAsset, const TCH
 		// If mesh's label has changed, update its name
 		if ( ExistingAsset->GetFName() != SourceAsset->GetFName() )
 		{
-			DestinationAsset->Rename( *SourceAsset->GetName(), DestinationPackage, REN_DontCreateRedirectors | REN_NonTransactional );
+			DestinationAsset->Rename( *AssetName, DestinationPackage, REN_DontCreateRedirectors | REN_NonTransactional );
 		}
 
 		if ( UStaticMesh* DestinationMesh = Cast< UStaticMesh >( DestinationAsset ) )
@@ -191,7 +192,7 @@ UObject* FDatasmithImporterImpl::PublicizeAsset( UObject* SourceAsset, const TCH
 	}
 	else
 	{
-		SourceAsset->Rename( *SourceAsset->GetName(), DestinationPackage, REN_DontCreateRedirectors | REN_NonTransactional );
+		SourceAsset->Rename( *AssetName, DestinationPackage, REN_DontCreateRedirectors | REN_NonTransactional );
 		DestinationAsset = SourceAsset;
 	}
 
