Datasmith: Add support for cloth actor

diff --git a/Engine/Plugins/Enterprise/DatasmithImporter/Source/DatasmithImporter/Private/DatasmithActorImporter.cpp b/Engine/Plugins/Enterprise/DatasmithImporter/Source/DatasmithImporter/Private/DatasmithActorImporter.cpp
--- a/Engine/Plugins/Enterprise/DatasmithImporter/Source/DatasmithImporter/Private/DatasmithActorImporter.cpp
+++ b/Engine/Plugins/Enterprise/DatasmithImporter/Source/DatasmithImporter/Private/DatasmithActorImporter.cpp
@@ -3,7 +3,9 @@
 #include "DatasmithActorImporter.h"
 
 #include "DatasmithCameraImporter.h"
+#include "DatasmithClothImporter.h"
 #include "DatasmithImportContext.h"
+#include "DatasmithImporterModule.h"
 #include "DatasmithImportOptions.h"
 #include "DatasmithLandscapeImporter.h"
 #include "DatasmithLightImporter.h"
@@ -285,6 +287,65 @@ UStaticMeshComponent* FDatasmithActorImporter::ImportStaticMeshComponent( FDatas
 	return StaticMeshComponent;
 }
 
+AActor* FDatasmithActorImporter::ImportClothActor(FDatasmithImportContext& ImportContext, const TSharedRef<IDatasmithClothActorElement>& ClothActorElement)
+{
+	AActor* ImportedActor = ImportActor(AActor::StaticClass(), ClothActorElement, ImportContext, ImportContext.Options->StaticMeshActorImportPolicy);
+
+	if (ImportedActor)
+	{
+		// Find imported resource
+		// #ue_ds_cloth_todo: Can be converted to a FDatasmithImporterUtils::FindAsset call when the actual cloth asset type is known to the importer
+		const TCHAR* ClothName = ClothActorElement->GetCloth();
+		UObject* ClothAsset = nullptr;
+
+		// search in imported resources
+		for (const auto& ImportedAssetPair : ImportContext.ImportedClothes) // Algo::Find
+		{
+			if (FCString::Stricmp(ImportedAssetPair.Key->GetName(), ClothName) == 0)
+			{
+				ClothAsset = ImportedAssetPair.Value;
+				break;
+			}
+		}
+
+		// search in existing scene asset
+		if (ClothAsset == nullptr)
+		{
+			check(ImportContext.SceneAsset);
+			if (TSoftObjectPtr<UObject>* ClothAssetPtr = ImportContext.SceneAsset->Clothes.Find(FName(ClothName)))
+			{
+				ClothAsset = ClothAssetPtr->LoadSynchronous();
+			}
+		}
+
+		if (IDatasmithImporterExt* Ext = IDatasmithImporterModule::Get().GetClothImporterExtension())
+		{
+			if (USceneComponent* ClothComponent = Ext->MakeClothComponent(ImportedActor, ClothAsset))
+			{
+				ImportedActor->AddInstanceComponent(ClothComponent);
+				ClothComponent->SetupAttachment(ImportedActor->GetRootComponent());
+
+				ClothComponent->bVisualizeComponent = true;
+				ClothComponent->RegisterComponent();
+			}
+			else
+			{
+				ImportContext.LogError(FText::Format(LOCTEXT("ClothComponentCreationFailure", "Cannot create cloth component for asset {0} reauired by actor {1}."), FText::FromString(ClothName), FText::FromString(ClothActorElement->GetLabel())));
+			}
+		}
+		else
+		{
+			ImportContext.LogError(FText::Format(LOCTEXT("ClothImporterExtensionIssue", "Cannot find ClothImporterExtension, cannot import actor {0}."), FText::FromString(ClothActorElement->GetLabel())));
+		}
+
+		if (ClothAsset == nullptr)
+		{
+			ImportContext.LogError(FText::Format(LOCTEXT("MissingClothAsset", "Cannot find cloth asset {0} for actor {1}."), FText::FromString(ClothName), FText::FromString(ClothActorElement->GetLabel())));
+		}
+	}
+	return ImportedActor;
+}
+
 void FDatasmithActorImporter::SetupStaticMeshComponent( FDatasmithImportContext& ImportContext, UStaticMeshComponent* StaticMeshComponent, const TSharedRef< IDatasmithMeshActorElement >& MeshActorElement )
 {
 	if ( !StaticMeshComponent )
@@ -406,7 +467,7 @@ AActor* FDatasmithActorImporter::ImportCustomActor(FDatasmithImportContext& Impo
 		{
 			ImportContext.LogWarning(FText::Format(LOCTEXT("BadPropertyKey", "Cannot find Property '{0}' on actor {1}.")
 				, FText::FromString(PropStr)
-				, FText::FromString(InCustomActorElement->GetName() )) 
+				, FText::FromString(InCustomActorElement->GetName() ))
 			);
 			continue;
 		}
@@ -417,7 +478,7 @@ AActor* FDatasmithActorImporter::ImportCustomActor(FDatasmithImportContext& Impo
 			ImportContext.LogWarning(FText::Format(LOCTEXT("BadPropertyValue", "Cannot properly load value '{0}' on property '{1}' for actor {2}.")
 				, FText::FromStringView(KeyValueProperty->GetValue())
 				, FText::FromString(PropStr)
-				, FText::FromStringView(InCustomActorElement->GetName())) 
+				, FText::FromStringView(InCustomActorElement->GetName()))
 			);
 		}
 	}
@@ -730,7 +791,7 @@ void FDatasmithActorImporter::OverrideStaticMeshActorMaterial( const FDatasmithI
 	}
 }
 
-AActor * FDatasmithActorImporter::ImportDecalActor(FDatasmithImportContext& ImportContext, const TSharedRef<IDatasmithDecalActorElement>& DecalActorElement, FDatasmithActorUniqueLabelProvider& UniqueNameProvider)
+AActor* FDatasmithActorImporter::ImportDecalActor(FDatasmithImportContext& ImportContext, const TSharedRef<IDatasmithDecalActorElement>& DecalActorElement, FDatasmithActorUniqueLabelProvider& UniqueNameProvider)
 {
 	if (ImportContext.Options->OtherActorImportPolicy == EDatasmithImportActorPolicy::Ignore)
 	{
diff --git a/Engine/Plugins/Enterprise/DatasmithImporter/Source/DatasmithImporter/Private/DatasmithImporter.cpp b/Engine/Plugins/Enterprise/DatasmithImporter/Source/DatasmithImporter/Private/DatasmithImporter.cpp
--- a/Engine/Plugins/Enterprise/DatasmithImporter/Source/DatasmithImporter/Private/DatasmithImporter.cpp
+++ b/Engine/Plugins/Enterprise/DatasmithImporter/Source/DatasmithImporter/Private/DatasmithImporter.cpp
@@ -45,6 +45,7 @@
 #include "Editor/UnrealEdEngine.h"
 #include "EditorLevelUtils.h"
 #include "Engine/Engine.h"
+#include "Engine/SkinnedAsset.h"
 #include "Engine/StaticMesh.h"
 #include "Engine/StaticMeshActor.h"
 #include "HAL/FileManager.h"
@@ -72,6 +73,7 @@
 #include "Serialization/ObjectReader.h"
 #include "Serialization/ObjectWriter.h"
 #include "Settings/EditorExperimentalSettings.h"
+#include "SkinnedAssetCompiler.h"
 #include "SourceControlOperations.h"
 #include "Templates/UniquePtr.h"
 #include "UObject/Package.h"
@@ -455,6 +457,7 @@ void FDatasmithImporter::ImportClothes(FDatasmithImportContext& ImportContext)
 						}
 					}
 
+					bool bShouldBuildAsset = false;
 					if (ExistingCloth)
 					{
 						if (ExistingCloth->GetOuter() != Outer)
@@ -480,9 +483,10 @@ void FDatasmithImporter::ImportClothes(FDatasmithImportContext& ImportContext)
 					else
 					{
 						ClothAsset = Ext->MakeClothAsset(Outer, *ClothName, ObjectFlags);
+						bShouldBuildAsset = true;
 					}
 
-					Ext->FillCloth(ClothAsset, ClothElement, DsCloth);
+					Ext->FillCloth(ClothAsset, ClothElement, DsCloth); // , bShouldBuildAsset);
 				}
 
 				ImportContext.ImportedClothes.Add(ClothElement, ClothAsset);
@@ -574,8 +578,19 @@ UStaticMesh* FDatasmithImporter::FinalizeStaticMesh( UStaticMesh* SourceStaticMe
 UObject* FDatasmithImporter::FinalizeCloth(UObject* SourceCloth, const TCHAR* FolderPath, UObject* ExistingCloth, TMap<UObject*, UObject*>* ReferencesToRemap)
 {
 	// #ue_ds_cloth_note FDatasmithImporter::FinalizeCloth
+	if (Cast<USkinnedAsset>(ExistingCloth))
+	{
+		// #ue_ds_cloth_todo review if this step is legit
+		FSkinnedAssetCompilingManager::Get().FinishCompilation({Cast<USkinnedAsset>(ExistingCloth)});
+	}
+
 	UObject* Asset = FDatasmithImporterImpl::FinalizeAsset(SourceCloth, FolderPath, ExistingCloth, ReferencesToRemap);
-// 	UChaosClothAsset* DestinationCloth = Cast<UChaosClothAsset>(Asset);
+
+	if (Cast<USkinnedAsset>(Asset))
+	{
+		// #ue_ds_cloth_todo review if this step is legit
+		FSkinnedAssetCompilingManager::Get().FinishCompilation({Cast<USkinnedAsset>(Asset)});
+	}
 	return Asset;
 }
 
@@ -990,7 +1005,6 @@ void FDatasmithImporter::ImportActors( FDatasmithImportContext& ImportContext )
 	}
 	// end of the hotfix
 
-
 	ADatasmithSceneActor* ImportSceneActor = ImportContext.ActorsContext.ImportSceneActor;
 
 	// Create a scene actor to import with if we don't have one
@@ -1084,13 +1098,15 @@ AActor* FDatasmithImporter::ImportActor( FDatasmithImportContext& ImportContext,
 	AActor* ImportedActor = nullptr;
 	if (ActorElement->IsA(EDatasmithElementType::HierarchicalInstanceStaticMesh))
 	{
-		TSharedRef< IDatasmithHierarchicalInstancedStaticMeshActorElement > HISMActorElement = StaticCastSharedRef< IDatasmithHierarchicalInstancedStaticMeshActorElement >( ActorElement );
-		ImportedActor =  FDatasmithActorImporter::ImportHierarchicalInstancedStaticMeshAsActor( ImportContext, HISMActorElement, UniqueNameProvider );
+		ImportedActor =  FDatasmithActorImporter::ImportHierarchicalInstancedStaticMeshAsActor( ImportContext, StaticCastSharedRef< IDatasmithHierarchicalInstancedStaticMeshActorElement >( ActorElement ), UniqueNameProvider );
 	}
 	else if (ActorElement->IsA(EDatasmithElementType::StaticMeshActor))
 	{
-		TSharedRef< IDatasmithMeshActorElement > MeshActorElement = StaticCastSharedRef< IDatasmithMeshActorElement >( ActorElement );
-		ImportedActor = FDatasmithActorImporter::ImportStaticMeshActor( ImportContext, MeshActorElement );
+		ImportedActor = FDatasmithActorImporter::ImportStaticMeshActor( ImportContext, StaticCastSharedRef< IDatasmithMeshActorElement >( ActorElement ) );
+	}
+	else if (ActorElement->IsA(EDatasmithElementType::ClothActor))
+	{
+		ImportedActor = FDatasmithActorImporter::ImportClothActor(ImportContext, StaticCastSharedRef<IDatasmithClothActorElement>( ActorElement ) );
 	}
 	else if (ActorElement->IsA(EDatasmithElementType::EnvironmentLight))
 	{
@@ -1346,9 +1362,9 @@ void FDatasmithImporter::FinalizeActors( FDatasmithImportContext& ImportContext,
 					continue;
 				}
 
-				const bool bActorIsRelatedToDestionScene = DestinationSceneActor->RelatedActors.Contains( SourceActorPair.Key );
+				const bool bActorIsRelatedToDestinationScene = DestinationSceneActor->RelatedActors.Contains( SourceActorPair.Key );
 				TSoftObjectPtr< AActor >& ExistingActorPtr = DestinationSceneActor->RelatedActors.FindOrAdd( SourceActorPair.Key );
-				const bool bShouldFinalizeActor = bShouldSpawnNonExistingActors || !bActorIsRelatedToDestionScene || ( ExistingActorPtr.Get() && !ExistingActorPtr.Get()->IsPendingKillPending() );
+				const bool bShouldFinalizeActor = bShouldSpawnNonExistingActors || !bActorIsRelatedToDestinationScene || ( ExistingActorPtr.Get() && !ExistingActorPtr.Get()->IsPendingKillPending() );
 
 				if ( bShouldFinalizeActor )
 				{
diff --git a/Engine/Plugins/Enterprise/DatasmithImporter/Source/DatasmithImporter/Public/DatasmithActorImporter.h b/Engine/Plugins/Enterprise/DatasmithImporter/Source/DatasmithImporter/Public/DatasmithActorImporter.h
--- a/Engine/Plugins/Enterprise/DatasmithImporter/Source/DatasmithImporter/Public/DatasmithActorImporter.h
+++ b/Engine/Plugins/Enterprise/DatasmithImporter/Source/DatasmithImporter/Public/DatasmithActorImporter.h
@@ -8,9 +8,10 @@
 
 class AActor;
 class AStaticMeshActor;
-struct FDatasmithImportContext;
+class FDatasmithActorUniqueLabelProvider;
 class IDatasmithActorElement;
 class IDatasmithCameraActorElement;
+class IDatasmithClothActorElement;
 class IDatasmithCustomActorElement;
 class IDatasmithDecalActorElement;
 class IDatasmithEnvironmentElement;
@@ -27,7 +28,7 @@ class UHierarchicalInstancedStaticMeshComponent;
 class USceneComponent;
 class UStaticMeshComponent;
 struct FCachedActorLabels;
-class FDatasmithActorUniqueLabelProvider;
+struct FDatasmithImportContext;
 
 class FDatasmithActorImporter
 {
@@ -71,6 +72,8 @@ public:
 	static AStaticMeshActor* ImportStaticMeshActor( FDatasmithImportContext& ImportContext, const TSharedRef< IDatasmithMeshActorElement >& MeshActorElement );
 	static UStaticMeshComponent* ImportStaticMeshComponent( FDatasmithImportContext& ImportContext, const TSharedRef< IDatasmithMeshActorElement >& MeshActorElement, UObject* Outer, FDatasmithActorUniqueLabelProvider& UniqueNameProvider );
 
+	static AActor* ImportClothActor( FDatasmithImportContext& ImportContext, const TSharedRef< IDatasmithClothActorElement >& ClothActorElement );
+
 	/**
 	 * Spawns a cine camera actor
 	 */
@@ -100,12 +103,12 @@ public:
 	 * Spawns a landscape actor
 	 */
 	static AActor* ImportLandscapeActor( FDatasmithImportContext& ImportContext, const TSharedRef< IDatasmithLandscapeElement >& LandscapeActorElement );
-	
+
 	/**
 	 * Spawns an actor with a hierarchical instanced static mesh component
 	 */
 	static AActor* ImportHierarchicalInstancedStaticMeshAsActor( FDatasmithImportContext& ImportContext, const TSharedRef< IDatasmithHierarchicalInstancedStaticMeshActorElement >& HierarchicalInstancedStatictMeshActorElement, FDatasmithActorUniqueLabelProvider& UniqueNameProvider );
-	
+
 	static UHierarchicalInstancedStaticMeshComponent* ImportHierarchicalInstancedStaticMeshComponent( FDatasmithImportContext& ImportContext,
 		const TSharedRef< IDatasmithHierarchicalInstancedStaticMeshActorElement >& HierarchicalInstancedStaticMeshActorElement, UObject* Outer, FDatasmithActorUniqueLabelProvider& UniqueNameProvider );
 
diff --git a/Engine/Plugins/Enterprise/DatasmithImporter/Source/DatasmithImporter/Public/DatasmithClothImporter.h b/Engine/Plugins/Enterprise/DatasmithImporter/Source/DatasmithImporter/Public/DatasmithClothImporter.h
--- a/Engine/Plugins/Enterprise/DatasmithImporter/Source/DatasmithImporter/Public/DatasmithClothImporter.h
+++ b/Engine/Plugins/Enterprise/DatasmithImporter/Source/DatasmithImporter/Public/DatasmithClothImporter.h
@@ -6,9 +6,11 @@
 #include "Templates/SharedPointer.h"
 #include "UObject/ObjectMacros.h"
 
+class AActor;
 class FDatasmithCloth;
 class FDatasmithClothPresetPropertySet;
 class IDatasmithClothElement;
+class USceneComponent;
 
 
 // #ue_ds_cloth_arch: Temp API
@@ -20,5 +22,6 @@ public:
 
 	virtual UObject* MakeClothPropertyAsset(UObject* Outer, const TCHAR* Name, EObjectFlags ObjectFlags) = 0;
 	virtual void FillPropertySet(UObject* PropertySetAsset, TSharedRef<IDatasmithClothElement> ClothElement, const FDatasmithClothPresetPropertySet& PropertySet) = 0;
-};
 
+	virtual USceneComponent* MakeClothComponent(AActor* ImportedActor, UObject* ImportedAsset) = 0;
+};
diff --git a/Engine/Plugins/Enterprise/DatasmithImporter/Source/DatasmithImporter/Public/DatasmithImportContext.h b/Engine/Plugins/Enterprise/DatasmithImporter/Source/DatasmithImporter/Public/DatasmithImportContext.h
--- a/Engine/Plugins/Enterprise/DatasmithImporter/Source/DatasmithImporter/Public/DatasmithImportContext.h
+++ b/Engine/Plugins/Enterprise/DatasmithImporter/Source/DatasmithImporter/Public/DatasmithImportContext.h
@@ -257,9 +257,9 @@ struct DATASMITHIMPORTER_API FDatasmithImportContext
 	/** Map of imported mesh for each mesh element */
 	TMap< TSharedRef< IDatasmithMeshElement >, UStaticMesh* > ImportedStaticMeshes;
 
-	TMap< TSharedRef< IDatasmithClothElement >, UObject* > ImportedClothes; // #ue_ds_cloth_todo: UChaosClothAsset
+	TMap< TSharedRef< IDatasmithClothElement >, UObject* > ImportedClothes; // #ue_ds_cloth_arch UChaosClothAsset
 
-	TArray< UObject* > ImportedClothPresets; // #ue_ds_cloth_todo: UChaosClothPreset // #ue_ds_cloth_todo: map with a dedicated element so that clothes can share parameters
+	TArray< UObject* > ImportedClothPresets; // #ue_ds_cloth_arch: UChaosClothPreset // #ue_ds_cloth_todo: map with a dedicated element so that clothes can share presets // #ue_ds_cloth_existing presets
 
 	/** Register IDatasmithMeshElement by their name so they can be searched faster */
 	TMap< FString, TSharedRef < IDatasmithMeshElement > > ImportedStaticMeshesByName;
diff --git a/Engine/Plugins/Enterprise/DatasmithImporter/Source/DatasmithTranslator/Public/DatasmithPayload.h b/Engine/Plugins/Enterprise/DatasmithImporter/Source/DatasmithTranslator/Public/DatasmithPayload.h
--- a/Engine/Plugins/Enterprise/DatasmithImporter/Source/DatasmithTranslator/Public/DatasmithPayload.h
+++ b/Engine/Plugins/Enterprise/DatasmithImporter/Source/DatasmithTranslator/Public/DatasmithPayload.h
@@ -20,7 +20,9 @@ struct FDatasmithMeshElementPayload : public FDatasmithPayload
 };
 
 
-// #ue_ds_cloth_doc IDatasmithScene API
+/**
+ * Describes a Cloth element payload, which is the actual data to be imported.
+ */
 struct FDatasmithClothElementPayload : public FDatasmithPayload
 {
 	FDatasmithCloth Cloth;
diff --git a/Engine/Source/Runtime/Datasmith/DatasmithCore/Private/DatasmithSceneElementsImpl.cpp b/Engine/Source/Runtime/Datasmith/DatasmithCore/Private/DatasmithSceneElementsImpl.cpp
--- a/Engine/Source/Runtime/Datasmith/DatasmithCore/Private/DatasmithSceneElementsImpl.cpp
+++ b/Engine/Source/Runtime/Datasmith/DatasmithCore/Private/DatasmithSceneElementsImpl.cpp
@@ -1503,3 +1503,9 @@ void FDatasmithSceneImpl::AttachActorToSceneRoot(const TSharedPtr< IDatasmithAct
 		DatasmithSceneImplInternal::AttachActorToSceneRoot(this, Child, AttachmentRule, FoundChildHierarchy);
 	}
 }
+
+FDatasmithClothActorElementImpl::FDatasmithClothActorElementImpl(const TCHAR* InName)
+	: FDatasmithActorElementImpl<IDatasmithClothActorElement>(InName, EDatasmithElementType::ClothActor)
+{
+	Store.RegisterParameter(Cloth, "Cloth");
+}
diff --git a/Engine/Source/Runtime/Datasmith/DatasmithCore/Private/DatasmithSceneElementsImpl.h b/Engine/Source/Runtime/Datasmith/DatasmithCore/Private/DatasmithSceneElementsImpl.h
--- a/Engine/Source/Runtime/Datasmith/DatasmithCore/Private/DatasmithSceneElementsImpl.h
+++ b/Engine/Source/Runtime/Datasmith/DatasmithCore/Private/DatasmithSceneElementsImpl.h
@@ -387,7 +387,9 @@ private:
 	TDatasmithReferenceArrayProxy<IDatasmithMaterialIDElement> MaterialSlots;
 };
 
-// #ue_ds_cloth_doc: Experimental Element that describes a cloth asset
+/*
+ * Experimental Element that describes a cloth asset.
+ */
 class FDatasmithClothElementImpl : public FDatasmithElementImpl< IDatasmithClothElement >
 {
 public:
@@ -1020,6 +1022,20 @@ private:
 	TReflected<FString> MaterialPathName;
 };
 
+
+class FDatasmithClothActorElementImpl : public FDatasmithActorElementImpl<IDatasmithClothActorElement>
+{
+public:
+	FDatasmithClothActorElementImpl(const TCHAR* InName);
+
+	virtual void SetCloth(const TCHAR* InCloth) override { Cloth = InCloth; }
+	virtual const TCHAR* GetCloth() const override { return *(FString&)Cloth; }
+
+private:
+	TReflected<FString> Cloth;
+};
+
+
 class FDatasmithEnvironmentElementImpl : public FDatasmithLightActorElementImpl< IDatasmithEnvironmentElement >
 {
 public:
diff --git a/Engine/Source/Runtime/Datasmith/DatasmithCore/Private/DatasmithSceneFactory.cpp b/Engine/Source/Runtime/Datasmith/DatasmithCore/Private/DatasmithSceneFactory.cpp
--- a/Engine/Source/Runtime/Datasmith/DatasmithCore/Private/DatasmithSceneFactory.cpp
+++ b/Engine/Source/Runtime/Datasmith/DatasmithCore/Private/DatasmithSceneFactory.cpp
@@ -47,6 +47,8 @@ TSharedPtr< IDatasmithElement > FDatasmithSceneFactory::CreateElement( EDatasmit
 		return CreateMesh( InName );
 	case EDatasmithElementType::Cloth:
 		return CreateCloth( InName );
+	case EDatasmithElementType::ClothActor:
+		return CreateClothActor( InName );
 	case EDatasmithElementType::StaticMeshActor:
 		return CreateMeshActor( InName );
 	case EDatasmithElementType::PointLight:
@@ -197,14 +199,19 @@ TSharedRef< IDatasmithMeshElement > FDatasmithSceneFactory::CreateMesh( const TC
 	return MakeShared< FDatasmithMeshElementImpl >( InName );
 }
 
+TSharedRef< IDatasmithMeshActorElement > FDatasmithSceneFactory::CreateMeshActor( const TCHAR* InName )
+{
+	return MakeShared< FDatasmithMeshActorElementImpl<> >( InName );
+}
+
 TSharedRef< IDatasmithClothElement > FDatasmithSceneFactory::CreateCloth(const TCHAR* InName)
 {
 	return MakeShared< FDatasmithClothElementImpl >( InName );
 }
 
-TSharedRef< IDatasmithMeshActorElement > FDatasmithSceneFactory::CreateMeshActor( const TCHAR* InName )
+TSharedRef< IDatasmithClothActorElement > FDatasmithSceneFactory::CreateClothActor( const TCHAR* InName )
 {
-	return MakeShared< FDatasmithMeshActorElementImpl<> >( InName );
+	return MakeShared< FDatasmithClothActorElementImpl >( InName );
 }
 
 TSharedRef< IDatasmithHierarchicalInstancedStaticMeshActorElement > FDatasmithSceneFactory::CreateHierarchicalInstanceStaticMeshActor(const TCHAR* InName)
diff --git a/Engine/Source/Runtime/Datasmith/DatasmithCore/Private/DatasmithSceneXmlReader.cpp b/Engine/Source/Runtime/Datasmith/DatasmithCore/Private/DatasmithSceneXmlReader.cpp
--- a/Engine/Source/Runtime/Datasmith/DatasmithCore/Private/DatasmithSceneXmlReader.cpp
+++ b/Engine/Source/Runtime/Datasmith/DatasmithCore/Private/DatasmithSceneXmlReader.cpp
@@ -21,8 +21,18 @@ FDatasmithSceneXmlReader::FDatasmithSceneXmlReader() = default;
 FDatasmithSceneXmlReader::~FDatasmithSceneXmlReader() = default;
 
 
-const TCHAR* ActorTagsView[] = { DATASMITH_ACTORNAME, DATASMITH_ACTORMESHNAME, DATASMITH_CAMERANAME, DATASMITH_LIGHTNAME,
-DATASMITH_CUSTOMACTORNAME, DATASMITH_LANDSCAPENAME, DATASMITH_POSTPROCESSVOLUME, DATASMITH_ACTORHIERARCHICALINSTANCEDMESHNAME, DATASMITH_DECALACTORNAME };
+const TCHAR* ActorTagsView[] = {
+	DATASMITH_ACTORNAME,
+	DATASMITH_ACTORHIERARCHICALINSTANCEDMESHNAME,
+	DATASMITH_ACTORMESHNAME,
+	DATASMITH_CAMERANAME,
+	DATASMITH_CLOTHACTORNAME,
+	DATASMITH_CUSTOMACTORNAME,
+	DATASMITH_DECALACTORNAME,
+	DATASMITH_LANDSCAPENAME,
+	DATASMITH_LIGHTNAME,
+	DATASMITH_POSTPROCESSVOLUME,
+};
 
 template<>
 float FDatasmithSceneXmlReader::ValueFromString< float >( const FString& InString ) const
@@ -655,6 +665,12 @@ void FDatasmithSceneXmlReader::ParseActor(FXmlNode* InNode, TSharedPtr<IDatasmit
 		ParseMeshActor(InNode, MeshElement, Scene);
 		InOutElement = MeshElement;
 	}
+	else if (InNode->GetTag() == DATASMITH_CLOTHACTORNAME)
+	{
+		TSharedPtr< IDatasmithClothActorElement> Element = FDatasmithSceneFactory::CreateClothActor(*InNode->GetAttribute(TEXT("name")));
+		ParseClothActor(InNode, Element, Scene);
+		InOutElement = Element;
+	}
 	else if (InNode->GetTag() == DATASMITH_LIGHTNAME)
 	{
 		TSharedPtr< IDatasmithLightActorElement > LightElement;
@@ -812,6 +828,19 @@ void FDatasmithSceneXmlReader::ParseMeshActor(FXmlNode* InNode, TSharedPtr<IData
 	}
 }
 
+void FDatasmithSceneXmlReader::ParseClothActor(FXmlNode* InNode, TSharedPtr<IDatasmithClothActorElement>& OutElement, TSharedRef< IDatasmithScene > Scene) const
+{
+	ParseElement(InNode, OutElement.ToSharedRef());
+
+	for (FXmlNode* ChildNode : InNode->GetChildrenNodes())
+	{
+		if (ChildNode->GetTag() == TEXT("Cloth"))
+		{
+			OutElement->SetCloth(*ChildNode->GetAttribute(TEXT("name")));
+		}
+	}
+}
+
 void FDatasmithSceneXmlReader::ParseHierarchicalInstancedStaticMeshActor(FXmlNode* InNode, TSharedPtr<IDatasmithHierarchicalInstancedStaticMeshActorElement>& OutElement, TSharedRef< IDatasmithScene > Scene) const
 {
 	TSharedPtr<IDatasmithMeshActorElement> OutElementAsMeshActor = OutElement;
diff --git a/Engine/Source/Runtime/Datasmith/DatasmithCore/Private/DatasmithSceneXmlWriter.cpp b/Engine/Source/Runtime/Datasmith/DatasmithCore/Private/DatasmithSceneXmlWriter.cpp
--- a/Engine/Source/Runtime/Datasmith/DatasmithCore/Private/DatasmithSceneXmlWriter.cpp
+++ b/Engine/Source/Runtime/Datasmith/DatasmithCore/Private/DatasmithSceneXmlWriter.cpp
@@ -47,6 +47,7 @@ public:
 	static void WriteActorTags(const TSharedPtr< IDatasmithActorElement >& ActorElement, FArchive& Archive, int32 Indent);
 	static void WriteActorChildren(const TSharedPtr< IDatasmithActorElement >& ActorElement, FArchive& Archive, int32 Indent);
 	static void WriteMeshActorElement(const TSharedPtr< IDatasmithMeshActorElement >& MeshActorElement, FArchive& Archive, int32 Indent);
+	static void WriteClothActorElement(const TSharedPtr< IDatasmithClothActorElement >& ClothActorElement, FArchive& Archive, int32 Indent);
 
 	// Write the start of the actor element (Open the xml element and add the essential child elements for the mesh actor)
 	static void WriteBeginOfMeshActorElement(const TSharedPtr<IDatasmithMeshActorElement>& MeshActorElement, const FString& ElementTypeString, FArchive& Archive, int32 Indent);
@@ -570,6 +571,10 @@ void FDatasmithSceneXmlWriterImpl::WriteActorElement(const TSharedPtr< IDatasmit
 	{
 		WriteMeshActorElement(StaticCastSharedPtr< IDatasmithMeshActorElement >(ActorElement), Archive, Indent);
 	}
+	else if ( ActorElement->IsA( EDatasmithElementType::ClothActor ) )
+	{
+		WriteClothActorElement(StaticCastSharedPtr< IDatasmithClothActorElement >(ActorElement), Archive, Indent);
+	}
 	else if ( ActorElement->IsA( EDatasmithElementType::Camera ) )
 	{
 		WriteCameraActorElement(StaticCastSharedPtr< IDatasmithCameraActorElement >(ActorElement), Archive, Indent);
@@ -677,6 +682,19 @@ void FDatasmithSceneXmlWriterImpl::WriteMeshActorElement(const TSharedPtr< IData
 	WriteEndOfMeshActorElement(ElementTypeString, Archive, Indent);
 }
 
+void FDatasmithSceneXmlWriterImpl::WriteClothActorElement(const TSharedPtr< IDatasmithClothActorElement >& Element, FArchive& Archive, int32 Indent)
+{
+	WriteIndent(Archive, Indent);
+	SerializeToArchive(Archive, FString(TEXT("<") DATASMITH_CLOTHACTORNAME TEXT(" name=\"")) + SanitizeXMLText(Element->GetName()) + TEXT("\""));
+	SerializeToArchive(Archive, GetLabelAndLayer(Element) + TEXT(">") LINE_TERMINATOR);
+
+	WriteIndent(Archive, Indent + 1);
+	SerializeToArchive(Archive, TEXT("<Cloth name=\"") + SanitizeXMLText(Element->GetCloth()) + TEXT("\"/>") LINE_TERMINATOR);
+
+	WriteIndent(Archive, Indent);
+	SerializeToArchive(Archive, TEXT("</") DATASMITH_CLOTHACTORNAME TEXT(">") LINE_TERMINATOR);
+}
+
 void FDatasmithSceneXmlWriterImpl::WriteBeginOfMeshActorElement(const TSharedPtr<IDatasmithMeshActorElement>& MeshActorElement, const FString& ElementTypeString, FArchive& Archive, int32 Indent)
 {
 	WriteIndent(Archive, Indent);
@@ -1094,7 +1112,7 @@ void FDatasmithSceneXmlWriterImpl::WriteCameraActorElement(const TSharedPtr< IDa
 {
 	WriteIndent(Archive, Indent);
 
-	FString XmlString = TEXT("<") + FString(DATASMITH_CAMERANAME) + TEXT(" name=\"") + SanitizeXMLText( CameraElement->GetName() ) + ("\"");
+	FString XmlString = TEXT("<") + FString(DATASMITH_CAMERANAME) + TEXT(" name=\"") + SanitizeXMLText( CameraElement->GetName() ) + TEXT("\"");
 
 	XmlString += GetLabelAndLayer(CameraElement) + TEXT(">") + LINE_TERMINATOR;
 
diff --git a/Engine/Source/Runtime/Datasmith/DatasmithCore/Public/DatasmithDefinitions.h b/Engine/Source/Runtime/Datasmith/DatasmithCore/Public/DatasmithDefinitions.h
--- a/Engine/Source/Runtime/Datasmith/DatasmithCore/Public/DatasmithDefinitions.h
+++ b/Engine/Source/Runtime/Datasmith/DatasmithCore/Public/DatasmithDefinitions.h
@@ -49,6 +49,7 @@ enum class EDatasmithElementType : uint64
 	MaterialExpressionInput        = 1ull << 32,
 	MaterialExpressionOutput       = 1ull << 33,
 	Cloth                          = 1ull << 34,
+	ClothActor                     = 1ull << 35,
 };
 
 ENUM_CLASS_FLAGS( EDatasmithElementType ); // Define bitwise operators for EDatasmithElementType
@@ -456,6 +457,7 @@ static const TCHAR* KeyValuePropertyTypeStrings[] = { TEXT("String"), TEXT("Colo
 #define DATASMITH_ACTORNAME						TEXT("Actor")
 
 #define DATASMITH_CLOTH							TEXT("Cloth")
+#define DATASMITH_CLOTHACTORNAME				TEXT("ClothActor")
 
 //ACTOR MESHES
 #define DATASMITH_ACTORMESHNAME					TEXT("ActorMesh")
diff --git a/Engine/Source/Runtime/Datasmith/DatasmithCore/Public/DatasmithSceneFactory.h b/Engine/Source/Runtime/Datasmith/DatasmithCore/Public/DatasmithSceneFactory.h
--- a/Engine/Source/Runtime/Datasmith/DatasmithCore/Public/DatasmithSceneFactory.h
+++ b/Engine/Source/Runtime/Datasmith/DatasmithCore/Public/DatasmithSceneFactory.h
@@ -44,9 +44,11 @@ public:
 
 	static TSharedRef< IDatasmithMeshElement > CreateMesh( const TCHAR* InName );
 
+	static TSharedRef< IDatasmithMeshActorElement > CreateMeshActor(const TCHAR* InName);
+
 	static TSharedRef< IDatasmithClothElement > CreateCloth( const TCHAR* InName );
 
-	static TSharedRef< IDatasmithMeshActorElement > CreateMeshActor( const TCHAR* InName );
+	static TSharedRef< IDatasmithClothActorElement > CreateClothActor(const TCHAR* InName);
 
 	static TSharedRef< IDatasmithHierarchicalInstancedStaticMeshActorElement > CreateHierarchicalInstanceStaticMeshActor( const TCHAR* InName );
 
diff --git a/Engine/Source/Runtime/Datasmith/DatasmithCore/Public/DatasmithSceneXmlReader.h b/Engine/Source/Runtime/Datasmith/DatasmithCore/Public/DatasmithSceneXmlReader.h
--- a/Engine/Source/Runtime/Datasmith/DatasmithCore/Public/DatasmithSceneXmlReader.h
+++ b/Engine/Source/Runtime/Datasmith/DatasmithCore/Public/DatasmithSceneXmlReader.h
@@ -47,6 +47,7 @@ private:
 	void ParseTexture(FXmlNode* InNode, FString& OutTextureFilename, FDatasmithTextureSampler& OutTextureUV) const;
 	void ParseActor(FXmlNode* InNode, TSharedPtr<IDatasmithActorElement>& InOutElement, TSharedRef< IDatasmithScene > Scene, TMap< FString, TSharedPtr<IDatasmithActorElement> >& Actors) const;
 	void ParseMeshActor(FXmlNode* InNode, TSharedPtr<IDatasmithMeshActorElement>& OutElement, TSharedRef< IDatasmithScene > Scene) const;
+	void ParseClothActor(FXmlNode* InNode, TSharedPtr<IDatasmithClothActorElement>& OutElement, TSharedRef< IDatasmithScene > Scene) const;
 	void ParseHierarchicalInstancedStaticMeshActor(FXmlNode* InNode, TSharedPtr<IDatasmithHierarchicalInstancedStaticMeshActorElement>& OutElement, TSharedRef< IDatasmithScene > Scene) const;
 	void ParseLight(FXmlNode* InNode, TSharedPtr<IDatasmithLightActorElement>& OutElement) const;
 	void ParseCamera(FXmlNode* InNode, TSharedPtr<IDatasmithCameraActorElement>& OutElement) const;
diff --git a/Engine/Source/Runtime/Datasmith/DatasmithCore/Public/IDatasmithSceneElements.h b/Engine/Source/Runtime/Datasmith/DatasmithCore/Public/IDatasmithSceneElements.h
--- a/Engine/Source/Runtime/Datasmith/DatasmithCore/Public/IDatasmithSceneElements.h
+++ b/Engine/Source/Runtime/Datasmith/DatasmithCore/Public/IDatasmithSceneElements.h
@@ -297,6 +297,13 @@ public:
 	virtual void SetStaticMeshPathName(const TCHAR* InStaticMeshPathName) = 0;
 };
 
+class DATASMITHCORE_API IDatasmithClothActorElement : public IDatasmithActorElement
+{
+public:
+	virtual void SetCloth(const TCHAR* Cloth) = 0;
+	virtual const TCHAR* GetCloth() const = 0;
+};
+
 class DATASMITHCORE_API IDatasmithHierarchicalInstancedStaticMeshActorElement : public IDatasmithMeshActorElement
 {
 public:
