[Datasmith] Modified DirectLinkExtension and ExternalSource modules to be runtime

diff --git a/Engine/Plugins/Enterprise/DatasmithImporter/DatasmithImporter.uplugin b/Engine/Plugins/Enterprise/DatasmithImporter/DatasmithImporter.uplugin
--- a/Engine/Plugins/Enterprise/DatasmithImporter/DatasmithImporter.uplugin
+++ b/Engine/Plugins/Enterprise/DatasmithImporter/DatasmithImporter.uplugin
@@ -17,7 +17,7 @@
 	"Modules": [
 		{
 			"Name": "DatasmithExternalSource",
-			"Type": "Editor",
+			"Type": "Runtime",
 			"LoadingPhase": "Default"
 		},
 		{
@@ -37,6 +37,11 @@
 		},
 		{
 			"Name": "DirectLinkExtension",
+			"Type": "Runtime",
+			"LoadingPhase": "Default"
+		},
+		{
+			"Name": "DirectLinkExtensionEditor",
 			"Type": "Editor",
 			"LoadingPhase": "Default"
 		},
diff --git a/Engine/Plugins/Enterprise/DatasmithImporter/Source/DatasmithExternalSource/DatasmithExternalSource.Build.cs b/Engine/Plugins/Enterprise/DatasmithImporter/Source/DatasmithExternalSource/DatasmithExternalSource.Build.cs
--- a/Engine/Plugins/Enterprise/DatasmithImporter/Source/DatasmithExternalSource/DatasmithExternalSource.Build.cs
+++ b/Engine/Plugins/Enterprise/DatasmithImporter/Source/DatasmithExternalSource/DatasmithExternalSource.Build.cs
@@ -31,6 +31,7 @@ namespace UnrealBuildTool.Rules
 					new string[]
 					{
 						"DesktopPlatform",
+						"DirectLinkExtensionEditor",
 						"Slate",
 						"UnrealEd",
 					}
diff --git a/Engine/Plugins/Enterprise/DatasmithImporter/Source/DatasmithImporter/DatasmithImporter.Build.cs b/Engine/Plugins/Enterprise/DatasmithImporter/Source/DatasmithImporter/DatasmithImporter.Build.cs
--- a/Engine/Plugins/Enterprise/DatasmithImporter/Source/DatasmithImporter/DatasmithImporter.Build.cs
+++ b/Engine/Plugins/Enterprise/DatasmithImporter/Source/DatasmithImporter/DatasmithImporter.Build.cs
@@ -18,6 +18,7 @@ namespace UnrealBuildTool.Rules
 					"Core",
 					"CoreUObject",
 					"DesktopPlatform",
+					"DirectLinkExtension",
 					"EditorFramework",
 					"EditorStyle",
 					"EditorScriptingUtilities",
@@ -67,7 +68,7 @@ namespace UnrealBuildTool.Rules
 					"DatasmithCore",
 					"DatasmithTranslator",
 					"DatasmithContentEditor",
-					"DirectLinkExtension",
+					"DirectLinkExtensionEditor",
 					"ExternalSource",
 				}
 			);
diff --git a/Engine/Plugins/Enterprise/DatasmithImporter/Source/DatasmithImporter/Private/DatasmithImporterModule.cpp b/Engine/Plugins/Enterprise/DatasmithImporter/Source/DatasmithImporter/Private/DatasmithImporterModule.cpp
--- a/Engine/Plugins/Enterprise/DatasmithImporter/Source/DatasmithImporter/Private/DatasmithImporterModule.cpp
+++ b/Engine/Plugins/Enterprise/DatasmithImporter/Source/DatasmithImporter/Private/DatasmithImporterModule.cpp
@@ -31,7 +31,7 @@
 #include "DataprepAssetInterface.h"
 #include "DataprepAssetUserData.h"
 #include "DataprepCoreUtils.h"
-#include "DirectLinkExtensionModule.h"
+#include "DirectLinkExtensionEditorModule.h"
 #include "DirectLinkUriResolver.h"
 #include "Editor.h"
 #include "EditorFramework/AssetImportData.h"
@@ -248,7 +248,7 @@ void FDatasmithImporterModule::OnClickedImportDirectLinkMenuEntry()
 {
 	using namespace UE::DatasmithImporter;
 
-	TSharedPtr<FDirectLinkExternalSource> ExternalSource = IDirectLinkExtensionModule::Get().DisplayDirectLinkSourcesDialog();
+	TSharedPtr<FDirectLinkExternalSource> ExternalSource = IDirectLinkExtensionEditorModule::Get().DisplayDirectLinkSourcesDialog();
 	if ( ExternalSource )
 	{
 		TFuture<TSharedPtr<IDatasmithScene>> DatasmithSceneFuture = ExternalSource->AsyncLoad();
@@ -775,7 +775,7 @@ void FDatasmithImporterModule::SetupDatasmithContentDelegates()
 					return false;
 				}
 
-				return IDirectLinkExtensionModule::Get().GetManager().SetAssetAutoReimport(Asset, bEnabled);
+				return IDirectLinkExtensionEditorModule::Get().GetManager().SetAssetAutoReimport(Asset, bEnabled);
 			});
 		SetAssetAutoReimportHandle = SetAssetAutoReimport.GetHandle();
 		DatasmithContentEditorModule.RegisterSetAssetAutoReimportHandler(MoveTemp(SetAssetAutoReimport));
@@ -805,7 +805,7 @@ void FDatasmithImporterModule::SetupDatasmithContentDelegates()
 
 				if (SourceUri.HasScheme(FDirectLinkUriResolver::GetDirectLinkScheme()))
 				{
-					return IDirectLinkExtensionModule::Get().GetManager().IsAssetAutoReimportEnabled(Asset);
+					return IDirectLinkExtensionEditorModule::Get().GetManager().IsAssetAutoReimportEnabled(Asset);
 				}
 
 				return false;
diff --git a/Engine/Plugins/Enterprise/DatasmithImporter/Source/DirectLinkExtension/DirectLinkExtension.Build.cs b/Engine/Plugins/Enterprise/DatasmithImporter/Source/DirectLinkExtension/DirectLinkExtension.Build.cs
--- a/Engine/Plugins/Enterprise/DatasmithImporter/Source/DirectLinkExtension/DirectLinkExtension.Build.cs
+++ b/Engine/Plugins/Enterprise/DatasmithImporter/Source/DirectLinkExtension/DirectLinkExtension.Build.cs
@@ -1,4 +1,5 @@
 // Copyright Epic Games, Inc. All Rights Reserved.
+using System;
 
 namespace UnrealBuildTool.Rules
 {
@@ -11,31 +12,30 @@ namespace UnrealBuildTool.Rules
 			PrivateDependencyModuleNames.AddRange(
 				new string[]
 				{
-					"ContentBrowserData",
-					"EditorStyle",
-					"EditorWidgets",
 					"Engine",
 					"InputCore",
-					"MainFrame",
-					"Projects",
-					"PropertyEditor",
-					"Slate",
-					"SlateCore",
-					"ToolMenus",
-					"UnrealEd",
 				}
 			);
 
 			PublicDependencyModuleNames.AddRange(
 				new string[]
 				{
-					"ContentBrowser",
 					"Core",
 					"CoreUObject",
 					"DirectLink",
 					"ExternalSource",
 				}
 			);
+
+			if (Target.Type == TargetType.Editor)
+			{
+				PrivateDependencyModuleNames.AddRange(
+					new string[]
+					{
+						"UnrealEd",
+					}
+				);
+			}
 		}
 	}
 }
\ No newline at end of file
diff --git a/Engine/Plugins/Enterprise/DatasmithImporter/Source/DirectLinkExtension/Private/AutoReimportManager.cpp b/Engine/Plugins/Enterprise/DatasmithImporter/Source/DirectLinkExtension/Private/AutoReimportManager.cpp
--- a/Engine/Plugins/Enterprise/DatasmithImporter/Source/DirectLinkExtension/Private/AutoReimportManager.cpp
+++ b/Engine/Plugins/Enterprise/DatasmithImporter/Source/DirectLinkExtension/Private/AutoReimportManager.cpp
@@ -6,7 +6,9 @@
 #include "Async/Async.h"
 #include "ExternalSource.h"
 #include "SourceUri.h"
+#if WITH_EDITOR
 #include "EditorReimportHandler.h"
+#endif
 #include "ExternalSourceModule.h"
 #include "Stats/Stats.h"
 #include "Misc/App.h"
@@ -236,7 +238,11 @@ namespace UE::DatasmithImporter
 		// Make sure we are not triggering a reimport on an asset that doesn't have a DirectLink source.
 		if (Uri.IsValid())
 		{
+#if WITH_EDITOR
 			FReimportManager::Instance()->Reimport(Asset, /*bAskForNewFileIfMissing*/ false, /*bShowNotification*/ true, /*PreferredReimportFile*/ TEXT(""), /*SpecifiedReimportHandler */ nullptr, /*SourceFileIndex*/ INDEX_NONE, /*bForceNewFile*/ false, /*bAutomated*/ true);
+#else
+			// TODO: Find a way to trigger re-import at runtime
+#endif
 		}
 		else
 		{
diff --git a/Engine/Plugins/Enterprise/DatasmithImporter/Source/DirectLinkExtension/Private/AutoReimportManager.h b/Engine/Plugins/Enterprise/DatasmithImporter/Source/DirectLinkExtension/Private/AutoReimportManager.h
--- a/Engine/Plugins/Enterprise/DatasmithImporter/Source/DirectLinkExtension/Private/AutoReimportManager.h
+++ b/Engine/Plugins/Enterprise/DatasmithImporter/Source/DirectLinkExtension/Private/AutoReimportManager.h
@@ -7,7 +7,7 @@
 #include "Delegates/IDelegateInstance.h"
 #include "Logging/LogMacros.h"
 #include "Templates/SharedPointer.h"
-#include "TickableEditorObject.h"
+#include "Tickable.h"
 #include "UObject/SoftObjectPtr.h"
 
 class UObject;
@@ -17,15 +17,17 @@ namespace UE::DatasmithImporter
 	class FExternalSource;
 	struct FAutoReimportInfo;
 
-	class FAutoReimportManager : public FTickableEditorObject, public TSharedFromThis<FAutoReimportManager>
+	class FAutoReimportManager : public FTickableGameObject, public TSharedFromThis<FAutoReimportManager>
 	{
 
 	public:
-		/* Begin FTickableEditorObject Interface */
+		/* Begin FTickableGameObject Interface */
 		virtual void Tick(float DeltaTime) override;
 		virtual bool IsAllowedToTick() const override { return PendingAutoReimportObjects.Num() > 0 || !PendingReimportQueue.IsEmpty(); }
 		virtual TStatId GetStatId() const override;
-		/* End FTickableEditorObject Interface */
+		virtual bool IsTickableWhenPaused() const override { return true; }
+		virtual bool IsTickableInEditor() const override { return true; }
+		/* End FTickableGameObject Interface */
 
 		bool IsAssetAutoReimportEnabled(UObject* InAsset) const;
 
diff --git a/Engine/Plugins/Enterprise/DatasmithImporter/Source/DirectLinkExtension/Private/DirectLinkAssetObserver.cpp b/Engine/Plugins/Enterprise/DatasmithImporter/Source/DirectLinkExtension/Private/DirectLinkAssetObserver.cpp
--- a/Engine/Plugins/Enterprise/DatasmithImporter/Source/DirectLinkExtension/Private/DirectLinkAssetObserver.cpp
+++ b/Engine/Plugins/Enterprise/DatasmithImporter/Source/DirectLinkExtension/Private/DirectLinkAssetObserver.cpp
@@ -18,7 +18,9 @@ namespace UE::DatasmithImporter
 		AssetRegistry.OnAssetRemoved().AddRaw(this, &FDirectLinkAssetObserver::AssetRemovedEvent);
 		AssetRegistry.OnAssetUpdated().AddRaw(this, &FDirectLinkAssetObserver::AssetUpdatedEvent);
 
+#if WITH_EDITOR
 		FCoreUObjectDelegates::OnObjectPropertyChanged.AddRaw(this, &FDirectLinkAssetObserver::OnObjectPropertyChanged);
+#endif
 	}
 
 	FDirectLinkAssetObserver::~FDirectLinkAssetObserver()
@@ -30,7 +32,9 @@ namespace UE::DatasmithImporter
 			AssetRegistry.OnAssetUpdated().RemoveAll(this);
 		}
 
+#if WITH_EDITOR
 		FCoreUObjectDelegates::OnObjectPropertyChanged.RemoveAll(this);
+#endif
 	}
 
 	void FDirectLinkAssetObserver::AssetRemovedEvent(const FAssetData& AssetData)
diff --git a/Engine/Plugins/Enterprise/DatasmithImporter/Source/DirectLinkExtension/Private/DirectLinkExtensionModule.cpp b/Engine/Plugins/Enterprise/DatasmithImporter/Source/DirectLinkExtension/Private/DirectLinkExtensionModule.cpp
--- a/Engine/Plugins/Enterprise/DatasmithImporter/Source/DirectLinkExtension/Private/DirectLinkExtensionModule.cpp
+++ b/Engine/Plugins/Enterprise/DatasmithImporter/Source/DirectLinkExtension/Private/DirectLinkExtensionModule.cpp
@@ -5,15 +5,9 @@
 #include "DirectLinkExternalSource.h"
 #include "DirectLinkManager.h"
 #include "DirectLinkUriResolver.h"
-#include "UI/DirectLinkExtensionUI.h"
-#include "UI/SDirectLinkAvailableSource.h"
 
-#include "ContentBrowserModule.h"
-#include "Editor.h"
 #include "ExternalSourceModule.h"
-#include "Interfaces/IMainFrameModule.h"
 #include "IUriManager.h"
-#include "ToolMenus.h"
 
 #define LOCTEXT_NAMESPACE "DirectLinkEditorModule"
 
@@ -28,8 +22,6 @@ namespace UE::DatasmithImporter
 		{
 			// This will instantiate the DirectLinkManager.
 			FDirectLinkManager::GetInstance();
-			DirectLinkExtensionUI = MakeUnique<FDirectLinkExtensionUI>();
-
 			IUriManager& UriManager = IExternalSourceModule::Get().GetManager();
 			UriManager.RegisterResolver(DirectLinkUriResolverName, MakeShared<FDirectLinkUriResolver>());
 		}
@@ -42,7 +34,6 @@ namespace UE::DatasmithImporter
 				UriManager.UnregisterResolver(DirectLinkUriResolverName);
 			}
 
-			DirectLinkExtensionUI.Reset();
 			FDirectLinkManager::ResetInstance();
 		}
 
@@ -51,42 +42,17 @@ namespace UE::DatasmithImporter
 			return FDirectLinkManager::GetInstance();
 		}
 
-		virtual TSharedPtr<FDirectLinkExternalSource> DisplayDirectLinkSourcesDialog() override;
-
-	private:
-		TUniquePtr<FDirectLinkExtensionUI> DirectLinkExtensionUI;
-	};
-
-	TSharedPtr<FDirectLinkExternalSource> FDirectLinkExtensionModule::DisplayDirectLinkSourcesDialog()
-	{
-		TSharedPtr<SWindow> ParentWindow;
-
-		if (FModuleManager::Get().IsModuleLoaded("MainFrame"))
+		virtual void OverwriteUriResolver(const TSharedRef<UE::DatasmithImporter::IUriResolver>& UriResolver)
 		{
-			IMainFrameModule& MainFrame = FModuleManager::LoadModuleChecked<IMainFrameModule>("MainFrame");
-			ParentWindow = MainFrame.GetParentWindow();
-		}
-
-		TSharedRef<SWindow> Window = SNew(SWindow)
-			.Title(LOCTEXT("DirectLinkEditorAvailableSourcesTitle", "DirectLink Available Sources"))
-			.SizingRule(ESizingRule::UserSized)
-			.AutoCenter(EAutoCenter::PreferredWorkArea)
-			.ClientSize(FVector2D(600, 200))
-			.SupportsMinimize(false);
+			IUriManager& UriManager = IExternalSourceModule::Get().GetManager();
 
-		TSharedPtr<SDirectLinkAvailableSource> AvailableSourceWindow;
-		Window->SetContent
-		(
-			SAssignNew(AvailableSourceWindow, SDirectLinkAvailableSource)
-			.WidgetWindow(Window)
-			.ProceedButtonLabel(FText(LOCTEXT("SelectLabel", "Select")))
-			.ProceedButtonTooltip(FText::GetEmpty())
-		);
+			// Unregister last registered URI resolver
+			UriManager.UnregisterResolver(DirectLinkUriResolverName);
 
-		FSlateApplication::Get().AddModalWindow(Window, ParentWindow);
+			UriManager.RegisterResolver(DirectLinkUriResolverName, UriResolver);
+		}
 
-		return AvailableSourceWindow->GetShouldProceed() ? AvailableSourceWindow->GetSelectedSource() : nullptr;
-	}
+	};
 }
 
 IMPLEMENT_MODULE(UE::DatasmithImporter::FDirectLinkExtensionModule, DirectLinkExtension);
diff --git a/Engine/Plugins/Enterprise/DatasmithImporter/Source/DirectLinkExtension/Private/DirectLinkManager.cpp b/Engine/Plugins/Enterprise/DatasmithImporter/Source/DirectLinkExtension/Private/DirectLinkManager.cpp
--- a/Engine/Plugins/Enterprise/DatasmithImporter/Source/DirectLinkExtension/Private/DirectLinkManager.cpp
+++ b/Engine/Plugins/Enterprise/DatasmithImporter/Source/DirectLinkExtension/Private/DirectLinkManager.cpp
@@ -10,14 +10,9 @@
 #include "AssetRegistry/AssetRegistryModule.h"
 #include "AssetRegistry/IAssetRegistry.h"
 #include "Async/Async.h"
-#include "EditorReimportHandler.h"
 #include "Misc/CoreDelegates.h"
 #include "UObject/UObjectGlobals.h"
 
-#if WITH_EDITOR
-#include "Editor.h"
-#endif //WITH_EDITOR
-
 #define LOCTEXT_NAMESPACE "DirectLinkManager"
 
 DEFINE_LOG_CATEGORY(LogDirectLinkManager);
diff --git a/Engine/Plugins/Enterprise/DatasmithImporter/Source/DirectLinkExtension/Private/DirectLinkManager.h b/Engine/Plugins/Enterprise/DatasmithImporter/Source/DirectLinkExtension/Private/DirectLinkManager.h
--- a/Engine/Plugins/Enterprise/DatasmithImporter/Source/DirectLinkExtension/Private/DirectLinkManager.h
+++ b/Engine/Plugins/Enterprise/DatasmithImporter/Source/DirectLinkExtension/Private/DirectLinkManager.h
@@ -93,6 +93,9 @@ namespace UE::DatasmithImporter
 #if WITH_EDITOR
 		virtual bool IsAssetAutoReimportEnabled(UObject* InAsset) const override { return AutoReimportManger->IsAssetAutoReimportEnabled(InAsset); }
 		virtual bool SetAssetAutoReimport(UObject* InAsset, bool bEnabled) override { return AutoReimportManger->SetAssetAutoReimport(InAsset, bEnabled); }
+#else
+		virtual bool IsAssetAutoReimportEnabled(UObject* InAsset) const override { return false; }
+		virtual bool SetAssetAutoReimport(UObject* InAsset, bool bEnabled) override { return false; }
 #endif //WITH_EDITOR
 		virtual TArray<TSharedRef<FDirectLinkExternalSource>> GetExternalSourceList() const override;
 		virtual void UnregisterDirectLinkExternalSource(FName InName) override;
diff --git a/Engine/Plugins/Enterprise/DatasmithImporter/Source/DirectLinkExtension/Private/DirectLinkSourceSearchFilter.cpp b/Engine/Plugins/Enterprise/DatasmithImporter/Source/DirectLinkExtension/Private/DirectLinkSourceSearchFilter.cpp
deleted file mode 100644
diff --git a/Engine/Plugins/Enterprise/DatasmithImporter/Source/DirectLinkExtension/Private/DirectLinkUriResolver.cpp b/Engine/Plugins/Enterprise/DatasmithImporter/Source/DirectLinkExtension/Private/DirectLinkUriResolver.cpp
--- a/Engine/Plugins/Enterprise/DatasmithImporter/Source/DirectLinkExtension/Private/DirectLinkUriResolver.cpp
+++ b/Engine/Plugins/Enterprise/DatasmithImporter/Source/DirectLinkExtension/Private/DirectLinkUriResolver.cpp
@@ -52,13 +52,6 @@ namespace UE::DatasmithImporter
 		return TOptional<FDirectLinkSourceDescription>();
 	}
 
-#if WITH_EDITOR
-	TSharedPtr<FExternalSource> FDirectLinkUriResolver::BrowseExternalSource(const FSourceUri& DefaultSourceUri) const
-	{
-		return IDirectLinkExtensionModule::Get().DisplayDirectLinkSourcesDialog();
-	}
-#endif //WITH_EDITOR
-
 	const FString& FDirectLinkUriResolver::GetDirectLinkScheme()
 	{
 		static FString Scheme(TEXT("directlink"));
diff --git a/Engine/Plugins/Enterprise/DatasmithImporter/Source/DirectLinkExtension/Private/UI/DirectLinkExtensionStyle.cpp b/Engine/Plugins/Enterprise/DatasmithImporter/Source/DirectLinkExtension/Private/UI/DirectLinkExtensionStyle.cpp
deleted file mode 100644
diff --git a/Engine/Plugins/Enterprise/DatasmithImporter/Source/DirectLinkExtension/Private/UI/DirectLinkExtensionStyle.h b/Engine/Plugins/Enterprise/DatasmithImporter/Source/DirectLinkExtension/Private/UI/DirectLinkExtensionStyle.h
deleted file mode 100644
diff --git a/Engine/Plugins/Enterprise/DatasmithImporter/Source/DirectLinkExtension/Private/UI/DirectLinkExtensionUI.h b/Engine/Plugins/Enterprise/DatasmithImporter/Source/DirectLinkExtension/Private/UI/DirectLinkExtensionUI.h
deleted file mode 100644
diff --git a/Engine/Plugins/Enterprise/DatasmithImporter/Source/DirectLinkExtension/Private/UI/SDirectLinkAssetThumbnailIndicator.cpp b/Engine/Plugins/Enterprise/DatasmithImporter/Source/DirectLinkExtension/Private/UI/SDirectLinkAssetThumbnailIndicator.cpp
deleted file mode 100644
diff --git a/Engine/Plugins/Enterprise/DatasmithImporter/Source/DirectLinkExtension/Private/UI/SDirectLinkAssetThumbnailIndicator.h b/Engine/Plugins/Enterprise/DatasmithImporter/Source/DirectLinkExtension/Private/UI/SDirectLinkAssetThumbnailIndicator.h
deleted file mode 100644
diff --git a/Engine/Plugins/Enterprise/DatasmithImporter/Source/DirectLinkExtension/Private/UI/SDirectLinkAvailableSource.cpp b/Engine/Plugins/Enterprise/DatasmithImporter/Source/DirectLinkExtension/Private/UI/SDirectLinkAvailableSource.cpp
deleted file mode 100644
diff --git a/Engine/Plugins/Enterprise/DatasmithImporter/Source/DirectLinkExtension/Private/UI/SDirectLinkAvailableSource.h b/Engine/Plugins/Enterprise/DatasmithImporter/Source/DirectLinkExtension/Private/UI/SDirectLinkAvailableSource.h
deleted file mode 100644
diff --git a/Engine/Plugins/Enterprise/DatasmithImporter/Source/DirectLinkExtension/Public/DirectLinkExtensionModule.h b/Engine/Plugins/Enterprise/DatasmithImporter/Source/DirectLinkExtension/Public/DirectLinkExtensionModule.h
--- a/Engine/Plugins/Enterprise/DatasmithImporter/Source/DirectLinkExtension/Public/DirectLinkExtensionModule.h
+++ b/Engine/Plugins/Enterprise/DatasmithImporter/Source/DirectLinkExtension/Public/DirectLinkExtensionModule.h
@@ -11,7 +11,7 @@
 
 namespace UE::DatasmithImporter
 {
-	class FDirectLinkExternalSource;
+	class IUriResolver;
 }
 
 namespace DirectLink
@@ -42,21 +42,19 @@ public:
 		return FModuleManager::Get().IsModuleLoaded(DIRECTLINKEXTENSION_MODULE_NAME);
 	}
 
+	static DirectLink::FEndpoint& GetEndpoint()
+	{
+		return Get().GetManager().GetEndpoint();
+	}
+
 	/**
-	 * Return the DirectLinkManager singleton.
+	 * Overwrite the UriResolver of the DirectLinkExtension module.
 	 */
-	virtual UE::DatasmithImporter::IDirectLinkManager& GetManager() const = 0;
+	virtual void OverwriteUriResolver(const TSharedRef<UE::DatasmithImporter::IUriResolver>& UriResolver) = 0;
 
 	/**
-	 * Spawn a dialog window prompting the user to select one available FDirectLinkExternalSource.
-	 * The list of displayed DirectLink source can be filtered.
-	 * @return The selected DirectLinkExternalSource, nullptr is returned if the dialog was canceled. 
+	 * Return the DirectLinkManager singleton.
 	 */
-	virtual TSharedPtr<UE::DatasmithImporter::FDirectLinkExternalSource> DisplayDirectLinkSourcesDialog() = 0;
-
-	static DirectLink::FEndpoint& GetEndpoint()
-	{
-		return Get().GetManager().GetEndpoint();
-	}
+	virtual UE::DatasmithImporter::IDirectLinkManager& GetManager() const = 0;
 };
 
diff --git a/Engine/Plugins/Enterprise/DatasmithImporter/Source/DirectLinkExtension/Public/DirectLinkSourceSearchFilter.h b/Engine/Plugins/Enterprise/DatasmithImporter/Source/DirectLinkExtension/Public/DirectLinkSourceSearchFilter.h
deleted file mode 100644
diff --git a/Engine/Plugins/Enterprise/DatasmithImporter/Source/DirectLinkExtension/Public/DirectLinkUriResolver.h b/Engine/Plugins/Enterprise/DatasmithImporter/Source/DirectLinkExtension/Public/DirectLinkUriResolver.h
--- a/Engine/Plugins/Enterprise/DatasmithImporter/Source/DirectLinkExtension/Public/DirectLinkUriResolver.h
+++ b/Engine/Plugins/Enterprise/DatasmithImporter/Source/DirectLinkExtension/Public/DirectLinkUriResolver.h
@@ -13,6 +13,8 @@ namespace DirectLink
 
 namespace UE::DatasmithImporter
 {
+	class FDirectLinkExternalSource;
+
 	struct FDirectLinkSourceDescription
 	{
 		FString ComputerName;
@@ -30,8 +32,11 @@ namespace UE::DatasmithImporter
 		virtual bool CanResolveUri(const FSourceUri& Uri) const override;
 		virtual FName GetScheme() const override { return FName(GetDirectLinkScheme()); }
 #if WITH_EDITOR
-		virtual TSharedPtr<FExternalSource> BrowseExternalSource(const FSourceUri& DefaultSourceUri) const override;
-#endif //WITH_EDITOR
+		virtual TSharedPtr<FExternalSource> BrowseExternalSource(const FSourceUri& DefaultSourceUri) const override
+		{
+			return {};
+		}
+#endif
 		// IUriResolver interface end
 
 		/**
diff --git a/Engine/Plugins/Enterprise/DatasmithImporter/Source/DirectLinkExtensionEditor/DirectLinkExtensionEditor.Build.cs b/Engine/Plugins/Enterprise/DatasmithImporter/Source/DirectLinkExtensionEditor/DirectLinkExtensionEditor.Build.cs
new file mode 100644
--- /dev/null
+++ b/Engine/Plugins/Enterprise/DatasmithImporter/Source/DirectLinkExtensionEditor/DirectLinkExtensionEditor.Build.cs
@@ -0,0 +1,49 @@
+// Copyright Epic Games, Inc. All Rights Reserved.
+using System;
+
+namespace UnrealBuildTool.Rules
+{
+	public class DirectLinkExtensionEditor : ModuleRules
+	{
+		public DirectLinkExtensionEditor(ReadOnlyTargetRules Target)
+			: base(Target)
+		{
+
+			PrivateDependencyModuleNames.AddRange(
+				new string[]
+				{
+					"Engine",
+					"InputCore",
+				}
+			);
+
+			PublicDependencyModuleNames.AddRange(
+				new string[]
+				{
+					"ContentBrowser",
+					"Core",
+					"CoreUObject",
+					"DirectLink",
+					"DirectLinkExtension",
+					"ExternalSource",
+				}
+			);
+
+			PrivateDependencyModuleNames.AddRange(
+				new string[]
+				{
+					"ContentBrowserData",
+					"EditorStyle",
+					"EditorWidgets",
+					"MainFrame",
+					"Projects",
+					"PropertyEditor",
+					"Slate",
+					"SlateCore",
+					"ToolMenus",
+					"UnrealEd",
+				}
+			);
+		}
+	}
+}
\ No newline at end of file
diff --git a/Engine/Plugins/Enterprise/DatasmithImporter/Source/DirectLinkExtensionEditor/Private/DirectLinkExtensionEditorModule.cpp b/Engine/Plugins/Enterprise/DatasmithImporter/Source/DirectLinkExtensionEditor/Private/DirectLinkExtensionEditorModule.cpp
new file mode 100644
--- /dev/null
+++ b/Engine/Plugins/Enterprise/DatasmithImporter/Source/DirectLinkExtensionEditor/Private/DirectLinkExtensionEditorModule.cpp
@@ -0,0 +1,75 @@
+// Copyright Epic Games, Inc. All Rights Reserved.
+
+#include "DirectLinkExtensionEditorModule.h"
+
+#include "DirectLinkUriResolverInEditor.h"
+#include "DirectLinkExtensionUI.h"
+#include "SDirectLinkAvailableSource.h"
+
+#include "ContentBrowserModule.h"
+#include "Editor.h"
+#include "Interfaces/IMainFrameModule.h"
+#include "ToolMenus.h"
+
+#define LOCTEXT_NAMESPACE "DirectLinkExtensionEditorModule"
+
+namespace UE::DatasmithImporter
+{
+	class FDirectLinkExtensionEditorModule : public IDirectLinkExtensionEditorModule
+	{
+		const FName DirectLinkUriResolverInEditorName = TEXT("DirectLinkUriResolverInEditorName");
+
+	public:
+		virtual void StartupModule() override
+		{
+			IDirectLinkExtensionModule::Get().OverwriteUriResolver(MakeShared<FDirectLinkUriResolverInEditor>());
+
+			DirectLinkExtensionUI = MakeUnique<FDirectLinkExtensionUI>();
+		}
+
+		virtual void ShutdownModule() override
+		{
+			DirectLinkExtensionUI.Reset();
+		}
+
+		virtual TSharedPtr<FDirectLinkExternalSource> DisplayDirectLinkSourcesDialog() override;
+
+	private:
+		TUniquePtr<FDirectLinkExtensionUI> DirectLinkExtensionUI;
+	};
+
+	TSharedPtr<FDirectLinkExternalSource> FDirectLinkExtensionEditorModule::DisplayDirectLinkSourcesDialog()
+	{
+		TSharedPtr<SWindow> ParentWindow;
+
+		if (FModuleManager::Get().IsModuleLoaded("MainFrame"))
+		{
+			IMainFrameModule& MainFrame = FModuleManager::LoadModuleChecked<IMainFrameModule>("MainFrame");
+			ParentWindow = MainFrame.GetParentWindow();
+		}
+
+		TSharedRef<SWindow> Window = SNew(SWindow)
+			.Title(LOCTEXT("DirectLinkEditorAvailableSourcesTitle", "DirectLink Available Sources"))
+			.SizingRule(ESizingRule::UserSized)
+			.AutoCenter(EAutoCenter::PreferredWorkArea)
+			.ClientSize(FVector2D(600, 200))
+			.SupportsMinimize(false);
+
+		TSharedPtr<SDirectLinkAvailableSource> AvailableSourceWindow;
+		Window->SetContent
+		(
+			SAssignNew(AvailableSourceWindow, SDirectLinkAvailableSource)
+			.WidgetWindow(Window)
+			.ProceedButtonLabel(FText(LOCTEXT("SelectLabel", "Select")))
+			.ProceedButtonTooltip(FText::GetEmpty())
+		);
+
+		FSlateApplication::Get().AddModalWindow(Window, ParentWindow);
+
+		return AvailableSourceWindow->GetShouldProceed() ? AvailableSourceWindow->GetSelectedSource() : nullptr;
+	}
+}
+
+IMPLEMENT_MODULE(UE::DatasmithImporter::FDirectLinkExtensionEditorModule, DirectLinkExtensionEditor)
+
+#undef LOCTEXT_NAMESPACE
diff --git a/Engine/Plugins/Enterprise/DatasmithImporter/Source/DirectLinkExtensionEditor/Private/DirectLinkExtensionStyle.cpp b/Engine/Plugins/Enterprise/DatasmithImporter/Source/DirectLinkExtensionEditor/Private/DirectLinkExtensionStyle.cpp
new file mode 100644
--- /dev/null
+++ b/Engine/Plugins/Enterprise/DatasmithImporter/Source/DirectLinkExtensionEditor/Private/DirectLinkExtensionStyle.cpp
@@ -0,0 +1,69 @@
+// Copyright Epic Games, Inc. All Rights Reserved.
+
+#include "DirectLinkExtensionStyle.h"
+
+#include "Styling/SlateStyleRegistry.h"
+#include "Brushes/SlateImageBrush.h"
+#include "Styling/AppStyle.h"
+#include "Interfaces/IPluginManager.h"
+
+#define IMAGE_PLUGIN_BRUSH( RelativePath, ... ) FSlateImageBrush( FDirectLinkExtensionStyle::InContent( RelativePath, TEXT(".png") ), __VA_ARGS__ )
+#define IMAGE_PLUGIN_BRUSH_SVG(RelativePath, ...) FSlateVectorImageBrush(FDirectLinkExtensionStyle::InContent(RelativePath, TEXT(".svg")), __VA_ARGS__)
+
+namespace UE::DatasmithImporter
+{
+	TUniquePtr<FSlateStyleSet> FDirectLinkExtensionStyle::StyleSet;
+
+	FString FDirectLinkExtensionStyle::InContent(const FString& RelativePath, const TCHAR* Extension)
+	{
+		static FString ContentDir = IPluginManager::Get().FindPlugin(TEXT("DatasmithImporter"))->GetContentDir();
+		return (ContentDir / RelativePath) + Extension;
+	}
+
+	FName FDirectLinkExtensionStyle::GetStyleSetName()
+	{
+		return FName(TEXT("DirectLinkExtensionStyle"));
+	}
+
+	void FDirectLinkExtensionStyle::Initialize()
+	{
+		// Only register once
+		if (StyleSet.IsValid())
+		{
+			return;
+		}
+
+		StyleSet = MakeUnique<FSlateStyleSet>(GetStyleSetName());
+		StyleSet->SetContentRoot(FPaths::EngineContentDir() / TEXT("Editor/Slate"));
+		StyleSet->SetCoreContentRoot(FPaths::EngineContentDir() / TEXT("Slate"));
+
+		// Const icon sizes
+		const FVector2D Icon16x16(16.0f, 16.0f);
+
+		// 16x16
+		StyleSet->Set("DirectLinkExtension.NotAvailable", new IMAGE_PLUGIN_BRUSH_SVG(TEXT("Icons/DirectLinkBadgeNotAvailable"), Icon16x16));
+		StyleSet->Set("DirectLinkExtension.OutOfSync", new IMAGE_PLUGIN_BRUSH_SVG(TEXT("Icons/DirectLinkBadgeOutOfSync"), Icon16x16));
+		StyleSet->Set("DirectLinkExtension.UpToDate", new IMAGE_PLUGIN_BRUSH_SVG(TEXT("Icons/DirectLinkBadgeUpToDate"), Icon16x16));
+		StyleSet->Set("DirectLinkExtension.AutoReimport", new IMAGE_PLUGIN_BRUSH_SVG(TEXT("Icons/DirectLinkBadgeAutoReimport"), Icon16x16));
+		StyleSet->Set("DirectLinkExtension.AutoReimportGrayscale", new IMAGE_PLUGIN_BRUSH_SVG(TEXT("Icons/DirectLinkAutoReimport_16"), Icon16x16));
+
+		FSlateStyleRegistry::RegisterSlateStyle(*StyleSet.Get());
+	};
+
+	void FDirectLinkExtensionStyle::Shutdown()
+	{
+		if (StyleSet.IsValid())
+		{
+			FSlateStyleRegistry::UnRegisterSlateStyle(*StyleSet.Get());
+			StyleSet.Reset();
+		}
+	}
+
+	const ISlateStyle& FDirectLinkExtensionStyle::Get()
+	{
+		return *StyleSet.Get();
+	}
+}
+
+#undef IMAGE_PLUGIN_BRUSH
+#undef IMAGE_PLUGIN_BRUSH_SVG
diff --git a/Engine/Plugins/Enterprise/DatasmithImporter/Source/DirectLinkExtensionEditor/Private/DirectLinkExtensionStyle.h b/Engine/Plugins/Enterprise/DatasmithImporter/Source/DirectLinkExtensionEditor/Private/DirectLinkExtensionStyle.h
new file mode 100644
--- /dev/null
+++ b/Engine/Plugins/Enterprise/DatasmithImporter/Source/DirectLinkExtensionEditor/Private/DirectLinkExtensionStyle.h
@@ -0,0 +1,30 @@
+// Copyright Epic Games, Inc. All Rights Reserved.
+
+#pragma once
+
+#include "CoreMinimal.h"
+#include "Styling/ISlateStyle.h"
+#include "Styling/SlateStyle.h"
+
+class ISlateStyle;
+class FSlateStyleSet;
+
+namespace UE::DatasmithImporter
+{
+	class FDirectLinkExtensionStyle
+	{
+	public:
+		static void Initialize();
+
+		static void Shutdown();
+
+		static const ISlateStyle& Get();
+
+		static FName GetStyleSetName();
+	private:
+		static FString InContent(const FString& RelativePath, const TCHAR* Extension);
+
+	private:
+		static TUniquePtr<FSlateStyleSet> StyleSet;
+	};
+}
diff --git a/Engine/Plugins/Enterprise/DatasmithImporter/Source/DirectLinkExtensionEditor/Private/DirectLinkExtensionUI.h b/Engine/Plugins/Enterprise/DatasmithImporter/Source/DirectLinkExtensionEditor/Private/DirectLinkExtensionUI.h
new file mode 100644
--- /dev/null
+++ b/Engine/Plugins/Enterprise/DatasmithImporter/Source/DirectLinkExtensionEditor/Private/DirectLinkExtensionUI.h
@@ -0,0 +1,68 @@
+// Copyright Epic Games, Inc. All Rights Reserved.
+
+#pragma once
+
+#include "DirectLinkExtensionStyle.h"
+#include "SDirectLinkAvailableSource.h"
+#include "SDirectLinkAssetThumbnailIndicator.h"
+
+#include "ContentBrowserModule.h"
+#include "CoreMinimal.h"
+#include "Interfaces/IMainFrameModule.h"
+#include "Modules/ModuleManager.h"
+
+#define LOCTEXT_NAMESPACE "DirectLinkExtensionUI"
+
+namespace UE::DatasmithImporter
+{
+	class FDirectLinkExtensionUI
+	{
+	public:
+		FDirectLinkExtensionUI()
+		{
+			FDirectLinkExtensionStyle::Initialize();
+
+			if (FContentBrowserModule* ContentBrowserModule = FModuleManager::Get().GetModulePtr<FContentBrowserModule>(TEXT("ContentBrowser")))
+			{
+				AssetViewExtraStateGeneratorHandle = ContentBrowserModule->AddAssetViewExtraStateGenerator(FAssetViewExtraStateGenerator(
+					FOnGenerateAssetViewExtraStateIndicators::CreateRaw(this, &FDirectLinkExtensionUI::OnGenerateAssetViewLockStateIcons),
+					FOnGenerateAssetViewExtraStateIndicators::CreateRaw(this, &FDirectLinkExtensionUI::OnGenerateAssetViewLockStateTooltip)
+				));
+			}
+		}
+
+		~FDirectLinkExtensionUI()
+		{
+			if (FContentBrowserModule* ContentBrowserModule = FModuleManager::Get().GetModulePtr<FContentBrowserModule>(TEXT("ContentBrowser")))
+			{
+				if (AssetViewExtraStateGeneratorHandle.IsValid())
+				{
+					ContentBrowserModule->RemoveAssetViewExtraStateGenerator(AssetViewExtraStateGeneratorHandle);
+				}
+			}
+
+			FDirectLinkExtensionStyle::Shutdown();
+		}
+
+	private:
+
+		TSharedRef<SWidget> OnGenerateAssetViewLockStateIcons(const FAssetData& AssetData)
+		{
+			return SNew(SDirectLinkAssetThumbnailIndicator)
+				.AssetData(AssetData);
+		}
+
+		TSharedRef<SWidget> OnGenerateAssetViewLockStateTooltip(const FAssetData& AssetData)
+		{
+			return SNew(SDirectLinkAssetThumbnailTooltip)
+				.AssetData(AssetData);
+		}
+
+		/* Delegate handle for asset direct link connection state indicator extensions.  */
+		FDelegateHandle AssetViewExtraStateGeneratorHandle;
+
+		TWeakPtr<SWindow> DirectLinkWindow;
+	};
+}
+
+#undef LOCTEXT_NAMESPACE
\ No newline at end of file
diff --git a/Engine/Plugins/Enterprise/DatasmithImporter/Source/DirectLinkExtensionEditor/Private/DirectLinkSourceSearchFilter.cpp b/Engine/Plugins/Enterprise/DatasmithImporter/Source/DirectLinkExtensionEditor/Private/DirectLinkSourceSearchFilter.cpp
new file mode 100644
--- /dev/null
+++ b/Engine/Plugins/Enterprise/DatasmithImporter/Source/DirectLinkExtensionEditor/Private/DirectLinkSourceSearchFilter.cpp
@@ -0,0 +1,51 @@
+// Copyright Epic Games, Inc. All Rights Reserved.
+
+#include "DirectLinkSourceSearchFilter.h"
+
+#include "DirectLinkUriResolver.h"
+
+#include "FrontendFilterBase.h"
+#include "SourceUri.h"
+
+#define LOCTEXT_NAMESPACE "UDirectLinkSourceSearchFilter"
+
+
+/** A filter that search for assets imported via DirectLink */
+class FFrontendFilter_DirectLinkSource : public FFrontendFilter
+{
+public:
+	FFrontendFilter_DirectLinkSource(TSharedPtr<FFrontendFilterCategory> InCategory)
+		: FFrontendFilter(InCategory)
+	{}
+
+	// FFrontendFilter implementation
+	virtual FLinearColor GetColor() const override { return FLinearColor::Green; }
+	virtual FString GetName() const override { return TEXT("DirectLinkSourceFilter"); }
+	virtual FText GetDisplayName() const override { return LOCTEXT("FrontendFilter_DirectLinkSource", "DirectLink Asset"); }
+	virtual FText GetToolTipText() const override { return LOCTEXT("FrontendFilter_DirectLinkSourceTooltip", "Show only assets have been imported with DirectLink."); }
+	// End of FFrontendFilter implementation
+
+	// IFilter implementation
+	virtual bool PassesFilter(FAssetFilterType InItem) const override;
+	// End of IFilter implementation
+};
+
+bool FFrontendFilter_DirectLinkSource::PassesFilter(FAssetFilterType InItem) const
+{
+	using namespace UE::DatasmithImporter;
+	FContentBrowserItemDataAttributeValue SourceUriAttribute = InItem.GetItemAttribute(FSourceUri::GetAssetDataTag());
+	if (SourceUriAttribute.IsValid())
+	{
+		FSourceUri SourceUri(SourceUriAttribute.GetValueString());
+		return SourceUri.HasScheme(FDirectLinkUriResolver::GetDirectLinkScheme());
+	}
+
+	return false;
+}
+
+void UDirectLinkSourceSearchFilter::AddFrontEndFilterExtensions(TSharedPtr<FFrontendFilterCategory> DefaultCategory, TArray< TSharedRef<class FFrontendFilter> >& InOutFilterList) const
+{
+	InOutFilterList.Add(MakeShareable(new FFrontendFilter_DirectLinkSource(DefaultCategory)));
+}
+
+#undef LOCTEXT_NAMESPACE
diff --git a/Engine/Plugins/Enterprise/DatasmithImporter/Source/DirectLinkExtensionEditor/Private/DirectLinkUriResolverInEditor.h b/Engine/Plugins/Enterprise/DatasmithImporter/Source/DirectLinkExtensionEditor/Private/DirectLinkUriResolverInEditor.h
new file mode 100644
--- /dev/null
+++ b/Engine/Plugins/Enterprise/DatasmithImporter/Source/DirectLinkExtensionEditor/Private/DirectLinkUriResolverInEditor.h
@@ -0,0 +1,22 @@
+// Copyright Epic Games, Inc. All Rights Reserved.
+
+#pragma once
+
+#include "DirectLinkUriResolver.h"
+
+#include "DirectLinkExtensionEditorModule.h"
+#include "DirectLinkExternalSource.h"
+
+namespace UE::DatasmithImporter
+{
+	class FDirectLinkUriResolverInEditor : public FDirectLinkUriResolver
+	{
+	public:
+		// IUriResolver interface begin
+		virtual TSharedPtr<FExternalSource> BrowseExternalSource(const FSourceUri& DefaultSourceUri) const override
+		{
+			return IDirectLinkExtensionEditorModule::Get().DisplayDirectLinkSourcesDialog();
+		}
+		// IUriResolver interface end
+	};
+}
\ No newline at end of file
diff --git a/Engine/Plugins/Enterprise/DatasmithImporter/Source/DirectLinkExtensionEditor/Private/SDirectLinkAssetThumbnailIndicator.cpp b/Engine/Plugins/Enterprise/DatasmithImporter/Source/DirectLinkExtensionEditor/Private/SDirectLinkAssetThumbnailIndicator.cpp
new file mode 100644
--- /dev/null
+++ b/Engine/Plugins/Enterprise/DatasmithImporter/Source/DirectLinkExtensionEditor/Private/SDirectLinkAssetThumbnailIndicator.cpp
@@ -0,0 +1,225 @@
+// Copyright Epic Games, Inc. All Rights Reserved.
+
+#include "SDirectLinkAssetThumbnailIndicator.h"
+
+#include "DirectLinkExtensionModule.h"
+#include "IDirectLinkManager.h"
+#include "DirectLinkExtensionStyle.h"
+
+#include "ExternalSource.h"
+#include "ExternalSourceModule.h"
+#include "SourceUri.h"
+#include "Widgets/Images/SImage.h"
+#include "Widgets/Text/STextBlock.h"
+#include "Widgets/SOverlay.h"
+
+#define LOCTEXT_NAMESPACE "DirectLinkAssetThumbnailIndicator"
+
+namespace UE::DatasmithImporter
+{
+	namespace ThumbnailIndicatorImpl
+	{
+		enum class EDirectLinkSourceStatus
+		{
+			None,
+			NotAvailable,
+			OutOfSync,
+			UpToDate,
+			AutoReimport,
+		};
+
+		FSourceUri TryGetSourceUri(const FAssetData& AssetData) 
+		{
+			const FAssetTagValueRef SourceUriTag = AssetData.TagsAndValues.FindTag(FSourceUri::GetAssetDataTag());
+			if (SourceUriTag.IsSet())
+			{
+				return FSourceUri(SourceUriTag.GetValue());
+			}
+
+			return FSourceUri();
+		}
+
+		EVisibility GetVisibility(const FAssetData& AssetData)
+		{
+			// If the asset has a DirectLink source, make the icon visible, collapsed/hidden otherwise.
+			const FSourceUri SourceUri = TryGetSourceUri(AssetData);
+			if (SourceUri.IsValid() && SourceUri.HasScheme(FDirectLinkUriResolver::GetDirectLinkScheme()))
+			{
+				return EVisibility::Visible;
+			}
+
+			return EVisibility::Collapsed;
+		}
+
+		FMD5Hash TryGetSourceHash(const FAssetData& AssetData)
+		{
+			FAssetTagValueRef SourceHashTag = AssetData.TagsAndValues.FindTag(FName(TEXT("SourceHash")));
+			if (SourceHashTag.IsSet())
+			{
+				FMD5Hash ParsedHash;
+				const FString SourceHashString(SourceHashTag.GetValue());
+				LexFromString(ParsedHash, *SourceHashString);
+
+				return ParsedHash;
+			}
+
+			return FMD5Hash();
+		}
+
+		const EDirectLinkSourceStatus GetSourceStatus(const FAssetData& AssetData)
+		{
+			const FSourceUri SourceUri = TryGetSourceUri(AssetData);
+			if (!SourceUri.IsValid() || !SourceUri.HasScheme(FDirectLinkUriResolver::GetDirectLinkScheme()))
+			{
+				return EDirectLinkSourceStatus::None;
+			}
+
+			// Only loaded asset can possibly be registered for auto-reimport.
+			const bool bLoadAsset = false;
+			if (UObject* Asset = AssetData.FastGetAsset(bLoadAsset))
+			{
+				if (IDirectLinkExtensionModule::Get().GetManager().IsAssetAutoReimportEnabled(Asset))
+				{
+					return EDirectLinkSourceStatus::AutoReimport;
+				}
+			}
+
+			const TSharedPtr<FExternalSource> ExternalSource = IExternalSourceModule::GetOrCreateExternalSource(SourceUri);
+			if (!ExternalSource.IsValid() || !ExternalSource->IsAvailable())
+			{
+				return EDirectLinkSourceStatus::NotAvailable;
+			}
+			else if (ExternalSource->IsOutOfSync())
+			{
+				// The scene loaded in ExternalSource is out of sync.
+				return EDirectLinkSourceStatus::OutOfSync;
+			}
+
+			const FMD5Hash SourceHash = TryGetSourceHash(AssetData);
+			const FMD5Hash ExternalSourceHash = ExternalSource->GetSourceHash();
+			if (!SourceHash.IsValid() || SourceHash != ExternalSourceHash)
+			{
+				//The source hash in the asset is different from the one in the external source, the asset is out of sync.
+				return EDirectLinkSourceStatus::OutOfSync;
+			}
+
+			return EDirectLinkSourceStatus::UpToDate;
+		}
+	}
+
+	const FSlateBrush* SDirectLinkAssetThumbnailIndicator::NotAvailableBrush = nullptr;
+	const FSlateBrush* SDirectLinkAssetThumbnailIndicator::OutOfSyncBrush = nullptr;
+	const FSlateBrush* SDirectLinkAssetThumbnailIndicator::UpToDateBrush = nullptr;
+	const FSlateBrush* SDirectLinkAssetThumbnailIndicator::AutoReimportBrush = nullptr;
+
+	void SDirectLinkAssetThumbnailIndicator::Construct(const FArguments& InArgs)
+	{
+		ObjectPropertyChangedHandle = FCoreUObjectDelegates::OnObjectPropertyChanged.AddSP(this, &SDirectLinkAssetThumbnailIndicator::OnObjectPropertyChanged);
+
+		CacheBrushes();
+		AssetData = InArgs._AssetData;
+		SetVisibility(MakeAttributeSP(this, &SDirectLinkAssetThumbnailIndicator::GetVisibility));
+
+		ChildSlot
+		[
+			SNew(SOverlay)
+			+ SOverlay::Slot() // The DirectLink image on top.
+			[
+				SNew(SImage)
+				.Image(this, &SDirectLinkAssetThumbnailIndicator::GetImageBrush)
+			]
+		];
+	}
+
+	SDirectLinkAssetThumbnailIndicator::~SDirectLinkAssetThumbnailIndicator()
+	{
+		FCoreUObjectDelegates::OnObjectPropertyChanged.Remove(ObjectPropertyChangedHandle);
+	}
+
+	void SDirectLinkAssetThumbnailIndicator::CacheBrushes()
+	{
+		if (NotAvailableBrush == nullptr)
+		{
+			NotAvailableBrush = FDirectLinkExtensionStyle::Get().GetBrush(TEXT("DirectLinkExtension.NotAvailable"));
+			OutOfSyncBrush = FDirectLinkExtensionStyle::Get().GetBrush(TEXT("DirectLinkExtension.OutOfSync"));
+			UpToDateBrush = FDirectLinkExtensionStyle::Get().GetBrush(TEXT("DirectLinkExtension.UpToDate"));
+			AutoReimportBrush = FDirectLinkExtensionStyle::Get().GetBrush(TEXT("DirectLinkExtension.AutoReimport"));
+		}
+	}
+
+	EVisibility SDirectLinkAssetThumbnailIndicator::GetVisibility() const
+	{
+		return ThumbnailIndicatorImpl::GetVisibility(AssetData);
+	}
+
+	const FSlateBrush* SDirectLinkAssetThumbnailIndicator::GetImageBrush() const
+	{
+		switch (ThumbnailIndicatorImpl::GetSourceStatus(AssetData))
+		{
+		case ThumbnailIndicatorImpl::EDirectLinkSourceStatus::NotAvailable:
+			return NotAvailableBrush;
+		case ThumbnailIndicatorImpl::EDirectLinkSourceStatus::OutOfSync:
+			return OutOfSyncBrush;
+		case ThumbnailIndicatorImpl::EDirectLinkSourceStatus::UpToDate:
+			return UpToDateBrush;
+		case ThumbnailIndicatorImpl::EDirectLinkSourceStatus::AutoReimport:
+			return AutoReimportBrush;
+		case ThumbnailIndicatorImpl::EDirectLinkSourceStatus::None:
+		default:
+			return nullptr;
+		}
+	}
+
+	void SDirectLinkAssetThumbnailIndicator::OnObjectPropertyChanged(UObject* ModifiedObject, struct FPropertyChangedEvent&)
+	{
+		if (ModifiedObject->IsA(UAssetImportData::StaticClass()))
+		{
+			ModifiedObject = ModifiedObject->GetOuter();
+		}
+
+		UObject* Asset = AssetData.FastGetAsset();
+		if (Asset && Asset == ModifiedObject)
+		{
+			// Refresh the asset data as it might have changed.
+			AssetData = FAssetData(Asset);
+		}
+	}
+
+	void SDirectLinkAssetThumbnailTooltip::Construct(const FArguments& InArgs)
+	{
+		AssetData = InArgs._AssetData;
+		SetVisibility(MakeAttributeSP(this, &SDirectLinkAssetThumbnailTooltip::GetTooltipVisibility));
+
+		ChildSlot
+		[
+			SNew(STextBlock)
+			.Text(this, &SDirectLinkAssetThumbnailTooltip::GetTooltipText)
+		];
+	}
+
+	EVisibility SDirectLinkAssetThumbnailTooltip::GetTooltipVisibility() const
+	{
+		return ThumbnailIndicatorImpl::GetVisibility(AssetData);
+	}
+
+	FText SDirectLinkAssetThumbnailTooltip::GetTooltipText() const
+	{
+		switch (ThumbnailIndicatorImpl::GetSourceStatus(AssetData))
+		{
+		case ThumbnailIndicatorImpl::EDirectLinkSourceStatus::NotAvailable:
+			return LOCTEXT("NotAvailableTooltip", "Direct Link source is not available.");
+		case ThumbnailIndicatorImpl::EDirectLinkSourceStatus::OutOfSync:
+			return LOCTEXT("OutOfSyncTooltip", "Asset is not synced to Direct Link source.");
+		case ThumbnailIndicatorImpl::EDirectLinkSourceStatus::UpToDate:
+			return LOCTEXT("UpToDateTooltip", "Asset is up-to-date with Direct Link source.");
+		case ThumbnailIndicatorImpl::EDirectLinkSourceStatus::AutoReimport:
+			return LOCTEXT("AutoReimportTooltip", "Asset will auto-reimport on Direct Link source change.");
+		case ThumbnailIndicatorImpl::EDirectLinkSourceStatus::None:
+		default:
+			return FText();
+		}
+	}
+
+}
+
+#undef LOCTEXT_NAMESPACE
diff --git a/Engine/Plugins/Enterprise/DatasmithImporter/Source/DirectLinkExtensionEditor/Private/SDirectLinkAssetThumbnailIndicator.h b/Engine/Plugins/Enterprise/DatasmithImporter/Source/DirectLinkExtensionEditor/Private/SDirectLinkAssetThumbnailIndicator.h
new file mode 100644
--- /dev/null
+++ b/Engine/Plugins/Enterprise/DatasmithImporter/Source/DirectLinkExtensionEditor/Private/SDirectLinkAssetThumbnailIndicator.h
@@ -0,0 +1,69 @@
+// Copyright Epic Games, Inc. All Rights Reserved.
+
+#pragma once
+
+#include "DirectLinkUriResolver.h"
+
+#include "DirectLinkUriResolver.h"
+
+#include "AssetRegistry/AssetData.h"
+#include "Framework/Application/SlateApplication.h"
+#include "Widgets/SCompoundWidget.h"
+
+struct FSlateBrush;
+
+namespace UE::DatasmithImporter
+{
+	class SDirectLinkAssetThumbnailIndicator : public SCompoundWidget
+	{
+	public:
+		SLATE_BEGIN_ARGS(SDirectLinkAssetThumbnailIndicator) {}
+			SLATE_ARGUMENT(FAssetData, AssetData)
+		SLATE_END_ARGS();
+
+		void Construct(const FArguments& InArgs);
+
+		virtual ~SDirectLinkAssetThumbnailIndicator();
+
+	private:
+		EVisibility GetVisibility() const;
+
+		const FSlateBrush* GetImageBrush() const;
+
+		void OnObjectPropertyChanged(UObject* ModifiedObject, struct FPropertyChangedEvent&);
+
+		static void CacheBrushes();
+
+		// Brushes used to render DirectLink icons.
+		static const FSlateBrush* NotAvailableBrush;
+		static const FSlateBrush* OutOfSyncBrush;
+		static const FSlateBrush* UpToDateBrush;
+		static const FSlateBrush* AutoReimportBrush;
+
+		/** Asset path for this indicator widget.*/
+		FAssetData AssetData;
+
+		FDelegateHandle ObjectPropertyChangedHandle;
+	};
+
+	class SDirectLinkAssetThumbnailTooltip : public SCompoundWidget
+	{
+	public:
+		SLATE_BEGIN_ARGS(SDirectLinkAssetThumbnailTooltip) {}
+			SLATE_ARGUMENT(FAssetData, AssetData)
+		SLATE_END_ARGS();
+
+		/**
+		 * Construct this widget.
+		 * @param InArgs Slate arguments
+		 */
+		void Construct(const FArguments& InArgs);
+		
+	private:
+		EVisibility GetTooltipVisibility() const;
+
+		FText GetTooltipText() const;
+
+		mutable FAssetData AssetData;
+	};
+}
diff --git a/Engine/Plugins/Enterprise/DatasmithImporter/Source/DirectLinkExtensionEditor/Private/SDirectLinkAvailableSource.cpp b/Engine/Plugins/Enterprise/DatasmithImporter/Source/DirectLinkExtensionEditor/Private/SDirectLinkAvailableSource.cpp
new file mode 100644
--- /dev/null
+++ b/Engine/Plugins/Enterprise/DatasmithImporter/Source/DirectLinkExtensionEditor/Private/SDirectLinkAvailableSource.cpp
@@ -0,0 +1,243 @@
+// Copyright Epic Games, Inc. All Rights Reserved.
+
+#include "SDirectLinkAvailableSource.h"
+#include "DirectLinkExternalSource.h"
+#include "DirectLinkUriResolver.h"
+#include "DirectLinkExtensionModule.h"
+#include "IDirectLinkManager.h"
+
+#include "Styling/AppStyle.h"
+#include "Styling/SlateBrush.h"
+#include "Widgets/Images/SImage.h"
+#include "Widgets/Input/SButton.h"
+#include "Widgets/Layout/SSpacer.h"
+#include "Widgets/Text/STextBlock.h"
+#include "Widgets/Views/SListView.h"
+#include "Widgets/Views/STableRow.h"
+
+#define LOCTEXT_NAMESPACE "DirectLinkAvailableSourceWindow"
+
+namespace UE::DatasmithImporter
+{
+	namespace SDirectLinkSourceWindowUtils
+	{
+		const FName ComputerColumnId("ComputerName");
+		const FName ProcessColumnId("ProcessName");
+		const FName EndpointColumnId("EndpointName");
+		const FName SourceColumnId("SourceName");
+	}
+
+	struct FDirectLinkExternalSourceInfo
+	{
+		FDirectLinkExternalSourceInfo(const TSharedRef<FDirectLinkExternalSource>& InExternalSource)
+			: ExternalSource(InExternalSource)
+		{
+			if (TOptional<FDirectLinkSourceDescription> ParsedSourceDescription = FDirectLinkUriResolver::TryParseDirectLinkUri(ExternalSource->GetSourceUri()))
+			{
+				SourceDescription = MoveTemp(ParsedSourceDescription.GetValue());
+			}
+		}
+
+		TSharedRef<FDirectLinkExternalSource> ExternalSource;
+		FDirectLinkSourceDescription SourceDescription;
+	};
+
+	class SDirectLinkExternalSourceInfoRow : public SMultiColumnTableRow<TSharedRef<FDirectLinkExternalSourceInfo>>
+	{
+	public:
+		SLATE_BEGIN_ARGS(SDirectLinkExternalSourceInfoRow)
+		{}
+		SLATE_END_ARGS()
+
+			void Construct(const FArguments& InArgs, TSharedRef<FDirectLinkExternalSourceInfo> InItem, const TSharedRef<STableViewBase>& InOwner)
+		{
+			ExternalSourceInfo = InItem;
+
+			FSuperRowType::FArguments Args = FSuperRowType::FArguments();
+			FSuperRowType::Construct(Args, InOwner);
+		}
+
+		virtual TSharedRef<SWidget> GenerateWidgetForColumn(const FName& ColumnName) override
+		{
+			const FDirectLinkSourceDescription& SourceDescription = ExternalSourceInfo->SourceDescription;
+
+			if (ColumnName == SDirectLinkSourceWindowUtils::ComputerColumnId)
+			{
+				return SNew(STextBlock)
+					.Text(FText::FromString(SourceDescription.ComputerName));
+			}
+			else if (ColumnName == SDirectLinkSourceWindowUtils::ProcessColumnId)
+			{
+				return SNew(STextBlock)
+					.Text(FText::FromString(SourceDescription.ExecutableName));
+			}
+			else if (ColumnName == SDirectLinkSourceWindowUtils::EndpointColumnId)
+			{
+				return SNew(STextBlock)
+					.Text(FText::FromString(SourceDescription.EndpointName));
+
+			}
+			else if (ColumnName == SDirectLinkSourceWindowUtils::SourceColumnId)
+			{
+				return SNew(STextBlock)
+					.Text(FText::FromString(SourceDescription.SourceName));
+			}
+
+			checkNoEntry();
+			return SNullWidget::NullWidget;
+		}
+
+	private:
+		TSharedPtr<FDirectLinkExternalSourceInfo> ExternalSourceInfo;
+	};
+
+	void SDirectLinkAvailableSource::Construct(const FArguments& InArgs)
+	{
+		Window = InArgs._WidgetWindow;
+
+		GenerateDirectLinkExternalSourceInfos();
+
+		TSharedRef<SHeaderRow> HeaderRow = SNew(SHeaderRow)
+			// Computer
+			+ SHeaderRow::Column(SDirectLinkSourceWindowUtils::ComputerColumnId)
+			.DefaultLabel(LOCTEXT("ComputerColumnLabel", "Computer"))
+			// Process
+			+SHeaderRow::Column(SDirectLinkSourceWindowUtils::ProcessColumnId)
+			.DefaultLabel(LOCTEXT("ProcessColumnLabel", "Process"))
+			// Endpoint
+			+ SHeaderRow::Column(SDirectLinkSourceWindowUtils::EndpointColumnId)
+			.DefaultLabel(LOCTEXT("EndpointColumnLabel", "Endpoint"))
+			// Source
+			+SHeaderRow::Column(SDirectLinkSourceWindowUtils::SourceColumnId)
+			.DefaultLabel(LOCTEXT("SourceColumnLabel", "Source"));
+
+		ChildSlot
+		[
+			SNew(SVerticalBox)
+			// No connection hint
+			+ SVerticalBox::Slot()
+			.Padding(2.f)
+			.FillHeight(TAttribute<float>(this, &SDirectLinkAvailableSource::GetNoConnectionHintFillHeight))
+			.VAlign(VAlign_Center)
+			[
+				SNew(STextBlock)
+				.Visibility(this, &SDirectLinkAvailableSource::GetNoConnectionHintVisibility)
+				.Text(LOCTEXT("ConnectionHintText", "No DirectLink source available."))
+				.AutoWrapText(true)
+				.Justification(ETextJustify::Center)
+				.Font(FCoreStyle::GetDefaultFontStyle("Regular", 13))
+			]
+
+			// Connection view
+			+ SVerticalBox::Slot()
+			.Padding(2.f)
+			.VAlign(VAlign_Fill)
+			[
+				SAssignNew(SourceListView, SListView< TSharedRef<FDirectLinkExternalSourceInfo>>)
+				.Visibility(this, &SDirectLinkAvailableSource::GetConnectionViewVisibility)
+				.ListItemsSource(&DirectLinkExternalSourceInfos)
+				.OnGenerateRow(this, &SDirectLinkAvailableSource::OnGenerateRow)
+				.OnSelectionChanged(this, &SDirectLinkAvailableSource::OnSelectionChanged)
+				.SelectionMode(ESelectionMode::Single)
+				.HeaderRow(HeaderRow)
+			]
+
+			+ SVerticalBox::Slot()
+			.Padding(2.f)
+			.AutoHeight()
+			.VAlign(VAlign_Bottom)
+			[
+				SNew(SHorizontalBox)
+				+SHorizontalBox::Slot()
+				.AutoWidth()
+				.HAlign(HAlign_Left)
+				[
+					SNew(SButton)
+					.ButtonStyle(FAppStyle::Get(), "SimpleButton")
+					.ContentPadding(1)
+					.ToolTipText(LOCTEXT("RefreshSourceListButtonTooltip", "Refresh the list of Direct Link sources."))
+					.OnClicked_Lambda([this]() { GenerateDirectLinkExternalSourceInfos(); SourceListView->RequestListRefresh(); return FReply::Handled(); })
+					[
+						SNew(SImage)
+						.Image(FAppStyle::GetBrush("Icons.Refresh"))
+					]
+				]
+
+				+SHorizontalBox::Slot()
+				.FillWidth(1.f)
+				[
+					SNew(SSpacer)
+				]
+
+				+ SHorizontalBox::Slot()
+				.Padding(5, 0)
+				.AutoWidth()
+				.HAlign(HAlign_Right)
+				[
+					SNew(SButton)
+					.HAlign(HAlign_Center)
+					.VAlign(VAlign_Center)
+					.Text(InArgs._ProceedButtonLabel)
+					.ToolTipText(InArgs._ProceedButtonTooltip)
+					.IsEnabled_Lambda([this]() { return SelectedSource.IsValid(); })
+					.OnClicked(this, &SDirectLinkAvailableSource::OnProceed)
+				]
+				+ SHorizontalBox::Slot()
+				.Padding(5, 0)
+				.AutoWidth()
+				.HAlign(HAlign_Right)
+				[
+					SNew(SButton)
+					.HAlign(HAlign_Center)
+					.VAlign(VAlign_Center)
+					.Text(FText(LOCTEXT("CancelButton", "Cancel")))
+					.OnClicked(this, &SDirectLinkAvailableSource::OnCancel)
+				]
+			]
+		];
+	}
+
+
+	float SDirectLinkAvailableSource::GetNoConnectionHintFillHeight() const
+	{
+		return DirectLinkExternalSourceInfos.Num() == 0 ? 1.f : 0.f;
+	}
+
+	::EVisibility SDirectLinkAvailableSource::GetNoConnectionHintVisibility() const
+	{
+		return DirectLinkExternalSourceInfos.Num() ? ::EVisibility::Hidden : ::EVisibility::Visible;
+	}
+
+	::EVisibility SDirectLinkAvailableSource::GetConnectionViewVisibility() const
+	{
+		return DirectLinkExternalSourceInfos.Num() ? ::EVisibility::Visible : ::EVisibility::Hidden;
+	}
+
+	void SDirectLinkAvailableSource::GenerateDirectLinkExternalSourceInfos()
+	{
+		DirectLinkExternalSourceInfos.Empty(DirectLinkExternalSourceInfos.Num());
+		for (const TSharedRef<FDirectLinkExternalSource>& ExternalSource : IDirectLinkExtensionModule::Get().GetManager().GetExternalSourceList())
+		{
+			DirectLinkExternalSourceInfos.Add(MakeShared<FDirectLinkExternalSourceInfo>(ExternalSource));
+		}
+	}
+
+	TSharedRef<ITableRow> SDirectLinkAvailableSource::OnGenerateRow(TSharedRef<FDirectLinkExternalSourceInfo> Item, const TSharedRef<STableViewBase>& Owner) const
+	{
+		return SNew(SDirectLinkExternalSourceInfoRow, Item, Owner);
+	}
+
+	void SDirectLinkAvailableSource::OnSelectionChanged(TSharedPtr<FDirectLinkExternalSourceInfo> SelectedValue, ESelectInfo::Type SelectionType)
+	{
+		if (SelectedValue)
+		{
+			SelectedSource = SelectedValue->ExternalSource;
+		}
+		else
+		{
+			SelectedSource = nullptr;
+		}
+	}
+}
+
+#undef LOCTEXT_NAMESPACE
diff --git a/Engine/Plugins/Enterprise/DatasmithImporter/Source/DirectLinkExtensionEditor/Private/SDirectLinkAvailableSource.h b/Engine/Plugins/Enterprise/DatasmithImporter/Source/DirectLinkExtensionEditor/Private/SDirectLinkAvailableSource.h
new file mode 100644
--- /dev/null
+++ b/Engine/Plugins/Enterprise/DatasmithImporter/Source/DirectLinkExtensionEditor/Private/SDirectLinkAvailableSource.h
@@ -0,0 +1,94 @@
+// Copyright Epic Games, Inc. All Rights Reserved.
+
+#pragma once
+
+#include "IDirectLinkManager.h"
+
+#include "Framework/Application/SlateApplication.h"
+#include "Widgets/SCompoundWidget.h"
+#include "Widgets/Views/SListView.h"
+
+namespace UE::DatasmithImporter
+{
+	class FExternalSource;
+
+	/**
+	 * Dialog used to make a selection of one available FDirectLinkExternalSource.
+	 */
+	class SDirectLinkAvailableSource : public SCompoundWidget
+	{
+	public:
+
+		SLATE_BEGIN_ARGS(SDirectLinkAvailableSource)
+		{}
+			SLATE_ARGUMENT(TSharedPtr<SWindow>, WidgetWindow)
+			SLATE_ARGUMENT(FText, ProceedButtonLabel)
+			SLATE_ARGUMENT(FText, ProceedButtonTooltip)
+		SLATE_END_ARGS()
+
+		void Construct(const FArguments& InArgs);
+
+		virtual ~SDirectLinkAvailableSource() = default;
+
+		virtual bool SupportsKeyboardFocus() const override { return true; }
+
+		virtual FReply OnKeyDown(const FGeometry& MyGeometry, const FKeyEvent& InKeyEvent) override
+		{
+			if (InKeyEvent.GetKey() == EKeys::Escape)
+			{
+				return OnCancel();
+			}
+
+			return FReply::Unhandled();
+		}
+
+		bool GetShouldProceed() const { return bShouldProceed; }
+
+		TSharedPtr<FDirectLinkExternalSource> GetSelectedSource() const { return SelectedSource; }
+
+	private:
+
+		void GenerateDirectLinkExternalSourceInfos();
+
+		TSharedRef<ITableRow> OnGenerateRow(TSharedRef<struct FDirectLinkExternalSourceInfo> Item, const TSharedRef<STableViewBase>& Owner) const;
+
+		void OnSelectionChanged(TSharedPtr<struct FDirectLinkExternalSourceInfo>, ESelectInfo::Type);
+
+		FReply OnCancel()
+		{
+			bShouldProceed = false;
+			if (Window.IsValid())
+			{
+				Window.Pin()->RequestDestroyWindow();
+			}
+			return FReply::Handled();
+		}
+
+		FReply OnProceed()
+		{
+			if (SelectedSource.IsValid() && Window.IsValid())
+			{
+				bShouldProceed = true;
+				Window.Pin()->RequestDestroyWindow();
+			}
+
+			return FReply::Handled();
+		}
+
+		float GetNoConnectionHintFillHeight() const;
+
+		::EVisibility GetNoConnectionHintVisibility() const;
+
+		::EVisibility GetConnectionViewVisibility() const;
+
+		bool bShouldProceed = false;
+
+		TWeakPtr< SWindow > Window;
+
+		TSharedPtr<FDirectLinkExternalSource> SelectedSource;
+
+		TSharedPtr<SListView<TSharedRef<struct FDirectLinkExternalSourceInfo>>> SourceListView;
+
+		TArray<TSharedRef<struct FDirectLinkExternalSourceInfo>> DirectLinkExternalSourceInfos;
+	};
+}
diff --git a/Engine/Plugins/Enterprise/DatasmithImporter/Source/DirectLinkExtensionEditor/Public/DirectLinkExtensionEditorModule.h b/Engine/Plugins/Enterprise/DatasmithImporter/Source/DirectLinkExtensionEditor/Public/DirectLinkExtensionEditorModule.h
new file mode 100644
--- /dev/null
+++ b/Engine/Plugins/Enterprise/DatasmithImporter/Source/DirectLinkExtensionEditor/Public/DirectLinkExtensionEditorModule.h
@@ -0,0 +1,54 @@
+// Copyright Epic Games, Inc. All Rights Reserved.
+
+#pragma once
+
+#include "DirectLinkExtensionModule.h"
+
+namespace UE::DatasmithImporter
+{
+	class FDirectLinkExternalSource;
+}
+
+#define DIRECTLINKEXTENSIONEDITOR_MODULE_NAME TEXT("DirectLinkExtensionEditor")
+
+class IDirectLinkExtensionEditorModule : public IDirectLinkExtensionModule
+{
+public:
+	/**
+	 * Singleton-like access to IDirectLinkExtensionEditorModule
+	 *
+	 * @return Returns IDirectLinkExtensionEditorModule singleton instance, loading the module on demand if needed
+	 */
+	static inline IDirectLinkExtensionEditorModule& Get()
+	{
+		return FModuleManager::LoadModuleChecked<IDirectLinkExtensionEditorModule>(DIRECTLINKEXTENSIONEDITOR_MODULE_NAME);
+	}
+
+	/**
+	 * Checks to see if this module is loaded and ready.
+	 *
+	 * @return True if the module is loaded and ready to use
+	 */
+	static inline bool IsAvailable()
+	{
+		return FModuleManager::Get().IsModuleLoaded(DIRECTLINKEXTENSIONEDITOR_MODULE_NAME);
+	}
+
+	virtual void OverwriteUriResolver(const TSharedRef<UE::DatasmithImporter::IUriResolver>& UriResolver)
+	{
+		IDirectLinkExtensionModule::Get().OverwriteUriResolver(UriResolver);
+	}
+
+	virtual UE::DatasmithImporter::IDirectLinkManager& GetManager() const override
+	{
+		return IDirectLinkExtensionModule::Get().GetManager();
+	}
+
+	/**
+	 * Spawn a dialog window prompting the user to select one available FDirectLinkExternalSource.
+	 * The list of displayed DirectLink source can be filtered.
+	 * @return The selected DirectLinkExternalSource, nullptr is returned if the dialog was canceled. 
+	 */
+	virtual TSharedPtr<UE::DatasmithImporter::FDirectLinkExternalSource> DisplayDirectLinkSourcesDialog() = 0;
+};
+
diff --git a/Engine/Plugins/Enterprise/DatasmithImporter/Source/DirectLinkExtensionEditor/Public/DirectLinkSourceSearchFilter.h b/Engine/Plugins/Enterprise/DatasmithImporter/Source/DirectLinkExtensionEditor/Public/DirectLinkSourceSearchFilter.h
new file mode 100644
--- /dev/null
+++ b/Engine/Plugins/Enterprise/DatasmithImporter/Source/DirectLinkExtensionEditor/Public/DirectLinkSourceSearchFilter.h
@@ -0,0 +1,26 @@
+// Copyright Epic Games, Inc. All Rights Reserved.
+
+#pragma once
+
+#include "CoreMinimal.h"
+#include "UObject/ObjectMacros.h"
+#include "ContentBrowserFrontEndFilterExtension.h"
+
+#include "DirectLinkSourceSearchFilter.generated.h"
+
+class FFrontendFilter;
+class FFrontendFilterCategory;
+
+/**
+ * Content Browser filter used to isolate assets imported via a Direct Link source.
+ */
+UCLASS()
+class UDirectLinkSourceSearchFilter : public UContentBrowserFrontEndFilterExtension
+{
+public:
+	GENERATED_BODY()
+
+	// UContentBrowserFrontEndFilterExtension interface
+	virtual void AddFrontEndFilterExtensions(TSharedPtr<FFrontendFilterCategory> DefaultCategory, TArray< TSharedRef<FFrontendFilter> >& InOutFilterList) const override;
+	// End of UContentBrowserFrontEndFilterExtension interface
+};
