Deprecated GetSceneImportOptions/SetSceneImportOptions IDatasmithTranslaor's methods using TStrongObjectPtr
Added GetSceneImportOptions/SetSceneImportOptions IDatasmithTranslaor's methods using TObjectPtr

#rnx

diff --git a/Engine/Plugins/Enterprise/DatasmithC4DImporter/Source/DatasmithC4DTranslator/Private/DatasmithC4DTranslator.cpp b/Engine/Plugins/Enterprise/DatasmithC4DImporter/Source/DatasmithC4DTranslator/Private/DatasmithC4DTranslator.cpp
--- a/Engine/Plugins/Enterprise/DatasmithC4DImporter/Source/DatasmithC4DTranslator/Private/DatasmithC4DTranslator.cpp
+++ b/Engine/Plugins/Enterprise/DatasmithC4DImporter/Source/DatasmithC4DTranslator/Private/DatasmithC4DTranslator.cpp
@@ -89,16 +89,16 @@ bool FDatasmithC4DTranslator::LoadLevelSequence(const TSharedRef<IDatasmithLevel
 	return false;
 }
 
-void FDatasmithC4DTranslator::GetSceneImportOptions(TArray<TStrongObjectPtr<UDatasmithOptionsBase>>& Options)
+void FDatasmithC4DTranslator::GetSceneImportOptions(TArray<TObjectPtr<UDatasmithOptionsBase>>& Options)
 {
-	Options.Add(GetOrCreateC4DImportOptions());
+	Options.Add(GetOrCreateC4DImportOptions().Get());
 }
 
-void FDatasmithC4DTranslator::SetSceneImportOptions(TArray<TStrongObjectPtr<UDatasmithOptionsBase>>& Options)
+void FDatasmithC4DTranslator::SetSceneImportOptions(const TArray<TObjectPtr<UDatasmithOptionsBase>>& Options)
 {
-	for (const TStrongObjectPtr<UDatasmithOptionsBase>& OptionPtr : Options)
+	for (const TObjectPtr<UDatasmithOptionsBase>& OptionPtr : Options)
 	{
-		if (UDatasmithC4DImportOptions* InImportOptions = Cast<UDatasmithC4DImportOptions>(OptionPtr.Get()))
+		if (UDatasmithC4DImportOptions* InImportOptions = Cast<UDatasmithC4DImportOptions>(OptionPtr))
 		{
 			ImportOptions.Reset(InImportOptions);
 		}
diff --git a/Engine/Plugins/Enterprise/DatasmithC4DImporter/Source/DatasmithC4DTranslator/Private/DatasmithC4DTranslator.h b/Engine/Plugins/Enterprise/DatasmithC4DImporter/Source/DatasmithC4DTranslator/Private/DatasmithC4DTranslator.h
--- a/Engine/Plugins/Enterprise/DatasmithC4DImporter/Source/DatasmithC4DTranslator/Private/DatasmithC4DTranslator.h
+++ b/Engine/Plugins/Enterprise/DatasmithC4DImporter/Source/DatasmithC4DTranslator/Private/DatasmithC4DTranslator.h
@@ -25,8 +25,8 @@ public:
 	virtual bool LoadStaticMesh(const TSharedRef<IDatasmithMeshElement> MeshElement, FDatasmithMeshElementPayload& OutMeshPayload) override;
 	virtual bool LoadLevelSequence(const TSharedRef<IDatasmithLevelSequenceElement> LevelSequenceElement, FDatasmithLevelSequencePayload& OutLevelSequencePayload) override;
 
-	virtual void GetSceneImportOptions(TArray<TStrongObjectPtr<UDatasmithOptionsBase>>& Options) override;
-	virtual void SetSceneImportOptions(TArray<TStrongObjectPtr<UDatasmithOptionsBase>>& Options) override;
+	virtual void GetSceneImportOptions(TArray<TObjectPtr<UDatasmithOptionsBase>>& Options) override;
+	virtual void SetSceneImportOptions(const TArray<TObjectPtr<UDatasmithOptionsBase>>& Options) override;
 
 	/** Event for when a C4D document is about to be opened for translation */
 	/*DECLARE_EVENT_TwoParams(FDatasmithC4DTranslator, FPreTranslateEvent, cineware::BaseDocument*, const FString&)
diff --git a/Engine/Plugins/Enterprise/DatasmithCADImporter/Source/DatasmithOpenNurbsTranslator/Private/DatasmithOpenNurbsTranslator.cpp b/Engine/Plugins/Enterprise/DatasmithCADImporter/Source/DatasmithOpenNurbsTranslator/Private/DatasmithOpenNurbsTranslator.cpp
--- a/Engine/Plugins/Enterprise/DatasmithCADImporter/Source/DatasmithOpenNurbsTranslator/Private/DatasmithOpenNurbsTranslator.cpp
+++ b/Engine/Plugins/Enterprise/DatasmithCADImporter/Source/DatasmithOpenNurbsTranslator/Private/DatasmithOpenNurbsTranslator.cpp
@@ -3330,15 +3330,15 @@ bool FOpenNurbsTranslatorImpl::LoadStaticMesh(const TSharedRef<IDatasmithMeshEle
 	return OutMeshPayload.LodMeshes.Num() > 0;
 }
 
-void FDatasmithOpenNurbsTranslator::SetSceneImportOptions(TArray<TStrongObjectPtr<UDatasmithOptionsBase>>& Options)
+void FDatasmithOpenNurbsTranslator::SetSceneImportOptions(const TArray<TObjectPtr<UDatasmithOptionsBase>>& Options)
 {
-	for (const TStrongObjectPtr<UDatasmithOptionsBase>& Option : Options)
+	for (const TObjectPtr<UDatasmithOptionsBase>& Option : Options)
 	{
-		if (UDatasmithImportOptions* DatasmithOptions = Cast<UDatasmithImportOptions>(Option.Get()))
+		if (UDatasmithImportOptions* DatasmithOptions = Cast<UDatasmithImportOptions>(Option))
 		{
 			BaseOptions = DatasmithOptions->BaseOptions;
 		}
-		else if (UDatasmithOpenNurbsImportOptions* OpenNurbsOptionsObj = Cast<UDatasmithOpenNurbsImportOptions>(Option.Get()))
+		else if (UDatasmithOpenNurbsImportOptions* OpenNurbsOptionsObj = Cast<UDatasmithOpenNurbsImportOptions>(Option))
 		{
 			OpenNurbsOptions = OpenNurbsOptionsObj->Options;
 		}
@@ -3351,9 +3351,9 @@ void FDatasmithOpenNurbsTranslator::SetSceneImportOptions(TArray<TStrongObjectPt
 	}
 }
 
-void FDatasmithOpenNurbsTranslator::GetSceneImportOptions(TArray<TStrongObjectPtr<UDatasmithOptionsBase>>& Options)
+void FDatasmithOpenNurbsTranslator::GetSceneImportOptions(TArray<TObjectPtr<UDatasmithOptionsBase>>& Options)
 {
-	TStrongObjectPtr<UDatasmithOpenNurbsImportOptions> OpenNurbsOptionsPtr = Datasmith::MakeOptions<UDatasmithOpenNurbsImportOptions>();
+	TObjectPtr<UDatasmithOpenNurbsImportOptions> OpenNurbsOptionsPtr = Datasmith::MakeOptionsObjectPtr<UDatasmithOpenNurbsImportOptions>();
 	if (ICADInterfacesModule::GetAvailability() == ECADInterfaceAvailability::Available)
 	{
 		OpenNurbsOptionsPtr->Options.Geometry = EDatasmithOpenNurbsBrepTessellatedSource::UseUnrealNurbsTessellation;
diff --git a/Engine/Plugins/Enterprise/DatasmithCADImporter/Source/DatasmithOpenNurbsTranslator/Private/DatasmithOpenNurbsTranslator.h b/Engine/Plugins/Enterprise/DatasmithCADImporter/Source/DatasmithOpenNurbsTranslator/Private/DatasmithOpenNurbsTranslator.h
--- a/Engine/Plugins/Enterprise/DatasmithCADImporter/Source/DatasmithOpenNurbsTranslator/Private/DatasmithOpenNurbsTranslator.h
+++ b/Engine/Plugins/Enterprise/DatasmithCADImporter/Source/DatasmithOpenNurbsTranslator/Private/DatasmithOpenNurbsTranslator.h
@@ -29,9 +29,9 @@ public:
 
 	virtual bool LoadStaticMesh(const TSharedRef<IDatasmithMeshElement> MeshElement, FDatasmithMeshElementPayload& OutMeshPayload) override;
 
-	virtual void SetSceneImportOptions(TArray<TStrongObjectPtr<UDatasmithOptionsBase>>& Options) override;
+	virtual void SetSceneImportOptions(const TArray<TObjectPtr<UDatasmithOptionsBase>>& Options) override;
 
-	virtual void GetSceneImportOptions(TArray<TStrongObjectPtr<UDatasmithOptionsBase>>& Options) override;
+	virtual void GetSceneImportOptions(TArray<TObjectPtr<UDatasmithOptionsBase>>& Options) override;
 
 private:
 	TSharedPtr<FOpenNurbsTranslatorImpl> Translator;
diff --git a/Engine/Plugins/Enterprise/DatasmithCADImporter/Source/DatasmithPLMXMLTranslator/Private/DatasmithPlmXmlTranslator.cpp b/Engine/Plugins/Enterprise/DatasmithCADImporter/Source/DatasmithPLMXMLTranslator/Private/DatasmithPlmXmlTranslator.cpp
--- a/Engine/Plugins/Enterprise/DatasmithCADImporter/Source/DatasmithPLMXMLTranslator/Private/DatasmithPlmXmlTranslator.cpp
+++ b/Engine/Plugins/Enterprise/DatasmithCADImporter/Source/DatasmithPLMXMLTranslator/Private/DatasmithPlmXmlTranslator.cpp
@@ -66,7 +66,7 @@ bool FDatasmithPlmXmlTranslator::LoadScene(TSharedRef<IDatasmithScene> OutScene)
     Importer = MakeUnique<FDatasmithPlmXmlImporter>(OutScene);
 
 	const FString& FilePath = GetSource().GetSourceFile();
-	if (!Importer->OpenFile(FilePath, GetSource(), CommonTessellationOptionsPtr->Options))
+	if (!Importer->OpenFile(FilePath, GetSource(), CommonTessellationOptions))
 	{
 		return false;
 	}
@@ -93,22 +93,22 @@ bool FDatasmithPlmXmlTranslator::LoadStaticMesh(const TSharedRef<IDatasmithMeshE
 	return false;
 }
 
-void FDatasmithPlmXmlTranslator::GetSceneImportOptions(TArray<TStrongObjectPtr<UDatasmithOptionsBase>>& Options)
+void FDatasmithPlmXmlTranslator::GetSceneImportOptions(TArray<TObjectPtr<UDatasmithOptionsBase>>& Options)
 {
-	if (!CommonTessellationOptionsPtr.IsValid())
-	{
-		CommonTessellationOptionsPtr = Datasmith::MakeOptions<UDatasmithCommonTessellationOptions>();
-	}
+	TObjectPtr<UDatasmithCommonTessellationOptions> CommonTessellationOptionsPtr = Datasmith::MakeOptionsObjectPtr<UDatasmithCommonTessellationOptions>();
+
+	CommonTessellationOptionsPtr->Options = CommonTessellationOptions;
+
 	Options.Add(CommonTessellationOptionsPtr);
 }
 
-void FDatasmithPlmXmlTranslator::SetSceneImportOptions(TArray<TStrongObjectPtr<UDatasmithOptionsBase>>& Options)
+void FDatasmithPlmXmlTranslator::SetSceneImportOptions(const TArray<TObjectPtr<UDatasmithOptionsBase>>& Options)
 {
-	for (const TStrongObjectPtr<UDatasmithOptionsBase>& OptionPtr : Options)
+	for (const TObjectPtr<UDatasmithOptionsBase>& OptionPtr : Options)
 	{
-		if (UDatasmithCommonTessellationOptions* TessellationOptionsObject = Cast<UDatasmithCommonTessellationOptions>(OptionPtr.Get()))
+		if (UDatasmithCommonTessellationOptions* TessellationOptionsObject = Cast<UDatasmithCommonTessellationOptions>(OptionPtr))
 		{
-			CommonTessellationOptionsPtr.Reset(TessellationOptionsObject);
+			CommonTessellationOptions = TessellationOptionsObject->Options;
 		}
 	}
 }
diff --git a/Engine/Plugins/Enterprise/DatasmithCADImporter/Source/DatasmithPLMXMLTranslator/Private/DatasmithPlmXmlTranslator.h b/Engine/Plugins/Enterprise/DatasmithCADImporter/Source/DatasmithPLMXMLTranslator/Private/DatasmithPlmXmlTranslator.h
--- a/Engine/Plugins/Enterprise/DatasmithCADImporter/Source/DatasmithPLMXMLTranslator/Private/DatasmithPlmXmlTranslator.h
+++ b/Engine/Plugins/Enterprise/DatasmithCADImporter/Source/DatasmithPLMXMLTranslator/Private/DatasmithPlmXmlTranslator.h
@@ -24,12 +24,12 @@ public:
 
 	virtual bool LoadStaticMesh(const TSharedRef<IDatasmithMeshElement> MeshElement, FDatasmithMeshElementPayload& OutMeshPayload) override;
 
-	virtual void GetSceneImportOptions(TArray<TStrongObjectPtr<UDatasmithOptionsBase>>& Options) override;
-	virtual void SetSceneImportOptions(TArray<TStrongObjectPtr<UDatasmithOptionsBase>>& Options) override;
+	virtual void GetSceneImportOptions(TArray<TObjectPtr<UDatasmithOptionsBase>>& Options) override;
+	virtual void SetSceneImportOptions(const TArray<TObjectPtr<UDatasmithOptionsBase>>& Options) override;
 	virtual bool IsSourceSupported(const FDatasmithSceneSource& Source) override;
 	//~ End IDatasmithTranslator interface
 
 private:
     TUniquePtr<FDatasmithPlmXmlImporter> Importer;
-	TStrongObjectPtr<UDatasmithCommonTessellationOptions> CommonTessellationOptionsPtr;
+	FDatasmithTessellationOptions CommonTessellationOptions;
 };
diff --git a/Engine/Plugins/Enterprise/DatasmithCADImporter/Source/DatasmithWireTranslator/Private/DatasmithWireTranslator.cpp b/Engine/Plugins/Enterprise/DatasmithCADImporter/Source/DatasmithWireTranslator/Private/DatasmithWireTranslator.cpp
--- a/Engine/Plugins/Enterprise/DatasmithCADImporter/Source/DatasmithWireTranslator/Private/DatasmithWireTranslator.cpp
+++ b/Engine/Plugins/Enterprise/DatasmithCADImporter/Source/DatasmithWireTranslator/Private/DatasmithWireTranslator.cpp
@@ -2334,7 +2334,7 @@ bool FDatasmithWireTranslator::LoadStaticMesh(const TSharedRef<IDatasmithMeshEle
 	return false;
 }
 
-void FDatasmithWireTranslator::SetSceneImportOptions(TArray<TStrongObjectPtr<UDatasmithOptionsBase>>& Options)
+void FDatasmithWireTranslator::SetSceneImportOptions(const TArray<TObjectPtr<UDatasmithOptionsBase>>& Options)
 {
 #ifdef USE_OPENMODEL
 	FParametricSurfaceTranslator::SetSceneImportOptions(Options);
diff --git a/Engine/Plugins/Enterprise/DatasmithCADImporter/Source/DatasmithWireTranslator/Private/DatasmithWireTranslator.h b/Engine/Plugins/Enterprise/DatasmithCADImporter/Source/DatasmithWireTranslator/Private/DatasmithWireTranslator.h
--- a/Engine/Plugins/Enterprise/DatasmithCADImporter/Source/DatasmithWireTranslator/Private/DatasmithWireTranslator.h
+++ b/Engine/Plugins/Enterprise/DatasmithCADImporter/Source/DatasmithWireTranslator/Private/DatasmithWireTranslator.h
@@ -32,7 +32,7 @@ public:
 	// End IDatasmithTranslator overrides
 
 	// Begin ADatasmithCoreTechTranslator overrides
-	virtual void SetSceneImportOptions(TArray<TStrongObjectPtr<UDatasmithOptionsBase>>& Options) override;
+	virtual void SetSceneImportOptions(const TArray<TObjectPtr<UDatasmithOptionsBase>>& Options) override;
 
 protected:
 	virtual void InitCommonTessellationOptions(FDatasmithTessellationOptions& TessellationOptions) override
diff --git a/Engine/Plugins/Enterprise/DatasmithCADImporter/Source/ParametricSurface/Private/ParametricSurfaceTranslator.cpp b/Engine/Plugins/Enterprise/DatasmithCADImporter/Source/ParametricSurface/Private/ParametricSurfaceTranslator.cpp
--- a/Engine/Plugins/Enterprise/DatasmithCADImporter/Source/ParametricSurface/Private/ParametricSurfaceTranslator.cpp
+++ b/Engine/Plugins/Enterprise/DatasmithCADImporter/Source/ParametricSurface/Private/ParametricSurfaceTranslator.cpp
@@ -10,7 +10,7 @@
 
 #include "Misc/FileHelper.h"
 
-void FParametricSurfaceTranslator::GetSceneImportOptions(TArray<TStrongObjectPtr<UDatasmithOptionsBase>>& Options)
+void FParametricSurfaceTranslator::GetSceneImportOptions(TArray<TObjectPtr<UDatasmithOptionsBase>>& Options)
 {
 	FString Extension = GetSource().GetSourceFileExtension();
 	if (Extension == TEXT("cgr") || Extension == TEXT("3dxml"))
@@ -18,18 +18,18 @@ void FParametricSurfaceTranslator::GetSceneImportOptions(TArray<TStrongObjectPtr
 		return;
 	}
 
-	TStrongObjectPtr<UDatasmithCommonTessellationOptions> CommonTessellationOptionsPtr = Datasmith::MakeOptions<UDatasmithCommonTessellationOptions>();
-	check(CommonTessellationOptionsPtr.IsValid());
+	TObjectPtr<UDatasmithCommonTessellationOptions> CommonTessellationOptionsPtr = Datasmith::MakeOptionsObjectPtr<UDatasmithCommonTessellationOptions>();
+	check(CommonTessellationOptionsPtr);
 	InitCommonTessellationOptions(CommonTessellationOptionsPtr->Options);
 
 	Options.Add(CommonTessellationOptionsPtr);
 }
 
-void FParametricSurfaceTranslator::SetSceneImportOptions(TArray<TStrongObjectPtr<UDatasmithOptionsBase>>& Options)
+void FParametricSurfaceTranslator::SetSceneImportOptions(const TArray<TObjectPtr<UDatasmithOptionsBase>>& Options)
 {
-	for (const TStrongObjectPtr<UDatasmithOptionsBase>& OptionPtr : Options)
+	for (const TObjectPtr<UDatasmithOptionsBase>& OptionPtr : Options)
 	{
-		if (UDatasmithCommonTessellationOptions* TessellationOptionsObject = Cast<UDatasmithCommonTessellationOptions>(OptionPtr.Get()))
+		if (UDatasmithCommonTessellationOptions* TessellationOptionsObject = Cast<UDatasmithCommonTessellationOptions>(OptionPtr))
 		{
 			CommonTessellationOptions = TessellationOptionsObject->Options;
 		}
diff --git a/Engine/Plugins/Enterprise/DatasmithCADImporter/Source/ParametricSurface/Public/ParametricSurfaceTranslator.h b/Engine/Plugins/Enterprise/DatasmithCADImporter/Source/ParametricSurface/Public/ParametricSurfaceTranslator.h
--- a/Engine/Plugins/Enterprise/DatasmithCADImporter/Source/ParametricSurface/Public/ParametricSurfaceTranslator.h
+++ b/Engine/Plugins/Enterprise/DatasmithCADImporter/Source/ParametricSurface/Public/ParametricSurfaceTranslator.h
@@ -29,8 +29,8 @@ public:
 	}
 
 	// Begin IDatasmithTranslator overrides
-	virtual void GetSceneImportOptions(TArray<TStrongObjectPtr<UDatasmithOptionsBase>>& Options) override;
-	virtual void SetSceneImportOptions(TArray<TStrongObjectPtr<UDatasmithOptionsBase>>& Options) override;
+	virtual void GetSceneImportOptions(TArray<TObjectPtr<UDatasmithOptionsBase>>& Options) override;
+	virtual void SetSceneImportOptions(const TArray<TObjectPtr<UDatasmithOptionsBase>>& Options) override;
 	// End IDatasmithTranslator overrides
 
 protected:
diff --git a/Engine/Plugins/Enterprise/DatasmithCADImporter/Source/ParametricSurfaceExtension/Private/ParametricRetessellateAction.cpp b/Engine/Plugins/Enterprise/DatasmithCADImporter/Source/ParametricSurfaceExtension/Private/ParametricRetessellateAction.cpp
--- a/Engine/Plugins/Enterprise/DatasmithCADImporter/Source/ParametricSurfaceExtension/Private/ParametricRetessellateAction.cpp
+++ b/Engine/Plugins/Enterprise/DatasmithCADImporter/Source/ParametricSurfaceExtension/Private/ParametricRetessellateAction.cpp
@@ -70,7 +70,7 @@ void FParametricRetessellateAction_Impl::ApplyOnAssets(const TArray<FAssetData>&
 		}
 	};
 
-	TStrongObjectPtr<UParametricRetessellateActionOptions> RetessellateOptions = Datasmith::MakeOptions<UParametricRetessellateActionOptions>();
+	TStrongObjectPtr<UParametricRetessellateActionOptions> RetessellateOptions(Datasmith::MakeOptions<UParametricRetessellateActionOptions>());
 
 	bool bSameOptionsForAll = false;
 	int32 NumAssetsToProcess = SelectedAssets.Num();
diff --git a/Engine/Plugins/Enterprise/DatasmithFBXImporter/Source/DatasmithDeltaGenTranslator/Private/DatasmithDeltaGenTranslator.cpp b/Engine/Plugins/Enterprise/DatasmithFBXImporter/Source/DatasmithDeltaGenTranslator/Private/DatasmithDeltaGenTranslator.cpp
--- a/Engine/Plugins/Enterprise/DatasmithFBXImporter/Source/DatasmithDeltaGenTranslator/Private/DatasmithDeltaGenTranslator.cpp
+++ b/Engine/Plugins/Enterprise/DatasmithFBXImporter/Source/DatasmithDeltaGenTranslator/Private/DatasmithDeltaGenTranslator.cpp
@@ -139,7 +139,7 @@ bool FDatasmithDeltaGenTranslator::LoadLevelSequence(const TSharedRef<IDatasmith
 	return false;
 }
 
-void FDatasmithDeltaGenTranslator::GetSceneImportOptions(TArray<TStrongObjectPtr<UDatasmithOptionsBase>>& Options)
+void FDatasmithDeltaGenTranslator::GetSceneImportOptions(TArray<TObjectPtr<UDatasmithOptionsBase>>& Options)
 {
 	if (!ImportOptions.IsValid())
 	{
@@ -149,14 +149,14 @@ void FDatasmithDeltaGenTranslator::GetSceneImportOptions(TArray<TStrongObjectPtr
 		ImportOptions->ResetPaths(FilePath, false);
 	}
 
-	Options.Add(ImportOptions);
+	Options.Add(ImportOptions.Get());
 }
 
-void FDatasmithDeltaGenTranslator::SetSceneImportOptions(TArray<TStrongObjectPtr<UDatasmithOptionsBase>>& Options)
+void FDatasmithDeltaGenTranslator::SetSceneImportOptions(const TArray<TObjectPtr<UDatasmithOptionsBase>>& Options)
 {
-	for (const TStrongObjectPtr<UDatasmithOptionsBase>& OptionPtr : Options)
+	for (const TObjectPtr<UDatasmithOptionsBase>& OptionPtr : Options)
 	{
-		if (UDatasmithDeltaGenImportOptions* InImportOptions = Cast<UDatasmithDeltaGenImportOptions>(OptionPtr.Get()))
+		if (UDatasmithDeltaGenImportOptions* InImportOptions = Cast<UDatasmithDeltaGenImportOptions>(OptionPtr))
 		{
 			ImportOptions.Reset(InImportOptions);
 		}
diff --git a/Engine/Plugins/Enterprise/DatasmithFBXImporter/Source/DatasmithDeltaGenTranslator/Private/DatasmithDeltaGenTranslator.h b/Engine/Plugins/Enterprise/DatasmithFBXImporter/Source/DatasmithDeltaGenTranslator/Private/DatasmithDeltaGenTranslator.h
--- a/Engine/Plugins/Enterprise/DatasmithFBXImporter/Source/DatasmithDeltaGenTranslator/Private/DatasmithDeltaGenTranslator.h
+++ b/Engine/Plugins/Enterprise/DatasmithFBXImporter/Source/DatasmithDeltaGenTranslator/Private/DatasmithDeltaGenTranslator.h
@@ -21,8 +21,8 @@ public:
 	virtual bool LoadStaticMesh(const TSharedRef<IDatasmithMeshElement> MeshElement, FDatasmithMeshElementPayload& OutMeshPayload) override;
 	virtual bool LoadLevelSequence(const TSharedRef<IDatasmithLevelSequenceElement> LevelSequenceElement, FDatasmithLevelSequencePayload& OutLevelSequencePayload) override;
 
-	virtual void GetSceneImportOptions(TArray<TStrongObjectPtr<UDatasmithOptionsBase>>& Options) override;
-	virtual void SetSceneImportOptions(TArray<TStrongObjectPtr<UDatasmithOptionsBase>>& Options) override;
+	virtual void GetSceneImportOptions(TArray<TObjectPtr<UDatasmithOptionsBase>>& Options) override;
+	virtual void SetSceneImportOptions(const TArray<TObjectPtr<UDatasmithOptionsBase>>& Options) override;
 
 private:
     TStrongObjectPtr<UDatasmithDeltaGenImportOptions> ImportOptions;
diff --git a/Engine/Plugins/Enterprise/DatasmithFBXImporter/Source/DatasmithVREDTranslator/Private/DatasmithVREDTranslator.cpp b/Engine/Plugins/Enterprise/DatasmithFBXImporter/Source/DatasmithVREDTranslator/Private/DatasmithVREDTranslator.cpp
--- a/Engine/Plugins/Enterprise/DatasmithFBXImporter/Source/DatasmithVREDTranslator/Private/DatasmithVREDTranslator.cpp
+++ b/Engine/Plugins/Enterprise/DatasmithFBXImporter/Source/DatasmithVREDTranslator/Private/DatasmithVREDTranslator.cpp
@@ -159,7 +159,7 @@ bool FDatasmithVREDTranslator::LoadLevelSequence(const TSharedRef<IDatasmithLeve
 	return false;
 }
 
-void FDatasmithVREDTranslator::GetSceneImportOptions(TArray<TStrongObjectPtr<UDatasmithOptionsBase>>& Options)
+void FDatasmithVREDTranslator::GetSceneImportOptions(TArray<TObjectPtr<UDatasmithOptionsBase>>& Options)
 {
 	if (!ImportOptions.IsValid())
 	{
@@ -169,14 +169,14 @@ void FDatasmithVREDTranslator::GetSceneImportOptions(TArray<TStrongObjectPtr<UDa
 		ImportOptions->ResetPaths(FilePath, false);
 	}
 
-	Options.Add(ImportOptions);
+	Options.Add(ImportOptions.Get());
 }
 
-void FDatasmithVREDTranslator::SetSceneImportOptions(TArray<TStrongObjectPtr<UDatasmithOptionsBase>>& Options)
+void FDatasmithVREDTranslator::SetSceneImportOptions(const TArray<TObjectPtr<UDatasmithOptionsBase>>& Options)
 {
-	for (const TStrongObjectPtr<UDatasmithOptionsBase>& OptionPtr : Options)
+	for (const TObjectPtr<UDatasmithOptionsBase>& OptionPtr : Options)
 	{
-		if (UDatasmithVREDImportOptions* InImportOptions = Cast<UDatasmithVREDImportOptions>(OptionPtr.Get()))
+		if (UDatasmithVREDImportOptions* InImportOptions = Cast<UDatasmithVREDImportOptions>(OptionPtr))
 		{
 			ImportOptions.Reset(InImportOptions);
 		}
diff --git a/Engine/Plugins/Enterprise/DatasmithFBXImporter/Source/DatasmithVREDTranslator/Private/DatasmithVREDTranslator.h b/Engine/Plugins/Enterprise/DatasmithFBXImporter/Source/DatasmithVREDTranslator/Private/DatasmithVREDTranslator.h
--- a/Engine/Plugins/Enterprise/DatasmithFBXImporter/Source/DatasmithVREDTranslator/Private/DatasmithVREDTranslator.h
+++ b/Engine/Plugins/Enterprise/DatasmithFBXImporter/Source/DatasmithVREDTranslator/Private/DatasmithVREDTranslator.h
@@ -21,8 +21,8 @@ public:
 	virtual bool LoadStaticMesh(const TSharedRef<IDatasmithMeshElement> MeshElement, FDatasmithMeshElementPayload& OutMeshPayload) override;
 	virtual bool LoadLevelSequence(const TSharedRef<IDatasmithLevelSequenceElement> LevelSequenceElement, FDatasmithLevelSequencePayload& OutLevelSequencePayload) override;
 
-	virtual void GetSceneImportOptions(TArray<TStrongObjectPtr<UDatasmithOptionsBase>>& Options) override;
-	virtual void SetSceneImportOptions(TArray<TStrongObjectPtr<UDatasmithOptionsBase>>& Options) override;
+	virtual void GetSceneImportOptions(TArray<TObjectPtr<UDatasmithOptionsBase>>& Options) override;
+	virtual void SetSceneImportOptions(const TArray<TObjectPtr<UDatasmithOptionsBase>>& Options) override;
 
 private:
     TStrongObjectPtr<UDatasmithVREDImportOptions> ImportOptions;
diff --git a/Engine/Plugins/Enterprise/DatasmithGLTFImporter/Source/DatasmithGLTFTranslator/Private/DatasmithGLTFTranslator.cpp b/Engine/Plugins/Enterprise/DatasmithGLTFImporter/Source/DatasmithGLTFTranslator/Private/DatasmithGLTFTranslator.cpp
--- a/Engine/Plugins/Enterprise/DatasmithGLTFImporter/Source/DatasmithGLTFTranslator/Private/DatasmithGLTFTranslator.cpp
+++ b/Engine/Plugins/Enterprise/DatasmithGLTFImporter/Source/DatasmithGLTFTranslator/Private/DatasmithGLTFTranslator.cpp
@@ -185,16 +185,16 @@ bool FDatasmithGLTFTranslator::LoadLevelSequence(const TSharedRef<IDatasmithLeve
 	return false;
 }
 
-void FDatasmithGLTFTranslator::GetSceneImportOptions(TArray<TStrongObjectPtr<UDatasmithOptionsBase>>& Options)
+void FDatasmithGLTFTranslator::GetSceneImportOptions(TArray<TObjectPtr<UDatasmithOptionsBase>>& Options)
 {
-	Options.Add(GetOrCreateGLTFImportOptions());
+	Options.Add(GetOrCreateGLTFImportOptions().Get());
 }
 
-void FDatasmithGLTFTranslator::SetSceneImportOptions(TArray<TStrongObjectPtr<UDatasmithOptionsBase>>& Options)
+void FDatasmithGLTFTranslator::SetSceneImportOptions(const TArray<TObjectPtr<UDatasmithOptionsBase>>& Options)
 {
-	for (const TStrongObjectPtr<UDatasmithOptionsBase>& OptionPtr : Options)
+	for (const TObjectPtr<UDatasmithOptionsBase>& OptionPtr : Options)
 	{
-		if (UDatasmithGLTFImportOptions* InImportOptions = Cast<UDatasmithGLTFImportOptions>(OptionPtr.Get()))
+		if (UDatasmithGLTFImportOptions* InImportOptions = Cast<UDatasmithGLTFImportOptions>(OptionPtr))
 		{
 			ImportOptions.Reset(InImportOptions);
 		}
diff --git a/Engine/Plugins/Enterprise/DatasmithGLTFImporter/Source/DatasmithGLTFTranslator/Private/DatasmithGLTFTranslator.h b/Engine/Plugins/Enterprise/DatasmithGLTFImporter/Source/DatasmithGLTFTranslator/Private/DatasmithGLTFTranslator.h
--- a/Engine/Plugins/Enterprise/DatasmithGLTFImporter/Source/DatasmithGLTFTranslator/Private/DatasmithGLTFTranslator.h
+++ b/Engine/Plugins/Enterprise/DatasmithGLTFImporter/Source/DatasmithGLTFTranslator/Private/DatasmithGLTFTranslator.h
@@ -20,8 +20,8 @@ public:
 	virtual bool LoadStaticMesh(const TSharedRef<IDatasmithMeshElement> MeshElement, FDatasmithMeshElementPayload& OutMeshPayload) override;
 	virtual bool LoadLevelSequence(const TSharedRef<IDatasmithLevelSequenceElement> LevelSequenceElement, FDatasmithLevelSequencePayload& OutLevelSequencePayload) override;
 
-	virtual void GetSceneImportOptions(TArray<TStrongObjectPtr<UDatasmithOptionsBase>>& Options) override;
-	virtual void SetSceneImportOptions(TArray<TStrongObjectPtr<UDatasmithOptionsBase>>& Options) override;
+	virtual void GetSceneImportOptions(TArray<TObjectPtr<UDatasmithOptionsBase>>& Options) override;
+	virtual void SetSceneImportOptions(const TArray<TObjectPtr<UDatasmithOptionsBase>>& Options) override;
 	//~ End IDatasmithTranslator interface
 
 private:
diff --git a/Engine/Plugins/Enterprise/DatasmithIFCImporter/Source/DatasmithIFCTranslator/Private/DatasmithIFCTranslator.cpp b/Engine/Plugins/Enterprise/DatasmithIFCImporter/Source/DatasmithIFCTranslator/Private/DatasmithIFCTranslator.cpp
--- a/Engine/Plugins/Enterprise/DatasmithIFCImporter/Source/DatasmithIFCTranslator/Private/DatasmithIFCTranslator.cpp
+++ b/Engine/Plugins/Enterprise/DatasmithIFCImporter/Source/DatasmithIFCTranslator/Private/DatasmithIFCTranslator.cpp
@@ -130,21 +130,21 @@ bool FDatasmithIFCTranslator::LoadLevelSequence(const TSharedRef<IDatasmithLevel
 	return false;
 }
 
-void FDatasmithIFCTranslator::GetSceneImportOptions(TArray<TStrongObjectPtr<UDatasmithOptionsBase>>& Options)
+void FDatasmithIFCTranslator::GetSceneImportOptions(TArray<TObjectPtr<UDatasmithOptionsBase>>& Options)
 {
 	if (!ImportOptions.IsValid())
 	{
 		ImportOptions = Datasmith::MakeOptions<UDatasmithIFCImportOptions>();
 	}
 
-	Options.Add(ImportOptions);
+	Options.Add(ImportOptions.Get());
 }
 
-void FDatasmithIFCTranslator::SetSceneImportOptions(TArray<TStrongObjectPtr<UDatasmithOptionsBase>>& Options)
+void FDatasmithIFCTranslator::SetSceneImportOptions(const TArray<TObjectPtr<UDatasmithOptionsBase>>& Options)
 {
-	for (const TStrongObjectPtr<UDatasmithOptionsBase>& OptionPtr : Options)
+	for (const TObjectPtr<UDatasmithOptionsBase>& OptionPtr : Options)
 	{
-		if (UDatasmithIFCImportOptions* InImportOptions = Cast<UDatasmithIFCImportOptions>(OptionPtr.Get()))
+		if (UDatasmithIFCImportOptions* InImportOptions = Cast<UDatasmithIFCImportOptions>(OptionPtr))
 		{
 			ImportOptions.Reset(InImportOptions);
 		}
diff --git a/Engine/Plugins/Enterprise/DatasmithIFCImporter/Source/DatasmithIFCTranslator/Private/DatasmithIFCTranslator.h b/Engine/Plugins/Enterprise/DatasmithIFCImporter/Source/DatasmithIFCTranslator/Private/DatasmithIFCTranslator.h
--- a/Engine/Plugins/Enterprise/DatasmithIFCImporter/Source/DatasmithIFCTranslator/Private/DatasmithIFCTranslator.h
+++ b/Engine/Plugins/Enterprise/DatasmithIFCImporter/Source/DatasmithIFCTranslator/Private/DatasmithIFCTranslator.h
@@ -20,8 +20,8 @@ public:
 	virtual bool LoadStaticMesh(const TSharedRef<IDatasmithMeshElement> MeshElement, FDatasmithMeshElementPayload& OutMeshPayload) override;
 	virtual bool LoadLevelSequence(const TSharedRef<IDatasmithLevelSequenceElement> LevelSequenceElement, FDatasmithLevelSequencePayload& OutLevelSequencePayload) override;
 
-	virtual void GetSceneImportOptions(TArray<TStrongObjectPtr<UDatasmithOptionsBase>>& Options) override;
-	virtual void SetSceneImportOptions(TArray<TStrongObjectPtr<UDatasmithOptionsBase>>& Options) override;
+	virtual void GetSceneImportOptions(TArray<TObjectPtr<UDatasmithOptionsBase>>& Options) override;
+	virtual void SetSceneImportOptions(const TArray<TObjectPtr<UDatasmithOptionsBase>>& Options) override;
 
 private:
     TStrongObjectPtr<UDatasmithIFCImportOptions> ImportOptions;
diff --git a/Engine/Plugins/Enterprise/DatasmithImporter/Source/DatasmithImporter/Private/DatasmithBlueprintLibrary.cpp b/Engine/Plugins/Enterprise/DatasmithImporter/Source/DatasmithImporter/Private/DatasmithBlueprintLibrary.cpp
--- a/Engine/Plugins/Enterprise/DatasmithImporter/Source/DatasmithImporter/Private/DatasmithBlueprintLibrary.cpp
+++ b/Engine/Plugins/Enterprise/DatasmithImporter/Source/DatasmithImporter/Private/DatasmithBlueprintLibrary.cpp
@@ -150,13 +150,13 @@ namespace DatasmithSceneElementUtil
 	{
 		static FDatasmithTessellationOptions DefaultTessellationOptions(0.3f, 0.0f, 30.0f, EDatasmithCADStitchingTechnique::StitchingSew);
 
-		TArray< TStrongObjectPtr<UDatasmithOptionsBase> > Options;
+		TArray< TObjectPtr<UDatasmithOptionsBase> > Options;
 		Translator->GetSceneImportOptions(Options);
 
 		bool bUpdateOptions = false;
-		for (TStrongObjectPtr<UDatasmithOptionsBase>& ObjectPtr : Options)
+		for (TObjectPtr<UDatasmithOptionsBase>& ObjectPtr : Options)
 		{
-			if (UDatasmithCommonTessellationOptions* TessellationOption = Cast<UDatasmithCommonTessellationOptions>(ObjectPtr.Get()))
+			if (UDatasmithCommonTessellationOptions* TessellationOption = Cast<UDatasmithCommonTessellationOptions>(ObjectPtr))
 			{
 				bUpdateOptions = true;
 				TessellationOption->Options = DefaultTessellationOptions;
@@ -431,7 +431,7 @@ namespace DatasmithSceneElementUtil
 							// Copy over the changes the user may have done on the options
 							ReImportSceneData->BaseOptions = ImportContext.Options->BaseOptions;
 
-							for (const TStrongObjectPtr<UDatasmithOptionsBase>& Option : ImportContext.AdditionalImportOptions)
+							for (const TObjectPtr<UDatasmithOptionsBase>& Option : ImportContext.AdditionalImportOptions)
 							{
 								UDatasmithOptionsBase* OptionObj = Option.Get();
 								OptionObj->Rename(nullptr, ReImportSceneData);
@@ -677,12 +677,6 @@ FDatasmithImportFactoryCreateFileResult UDatasmithSceneElement::ImportScene(cons
 {
 	FDatasmithImportFactoryCreateFileResult Result;
 
-	if (this == nullptr)
-	{
-		UE_LOG(LogDatasmithImport, Error, TEXT("Invalid State. Ensure ConstructDatasmithSceneFromFile has been called."));
-		return Result;
-	}
-
 	if (bMultifile)
 	{
 		TArray<FString> FilesNotProcessed;
@@ -749,7 +743,7 @@ FDatasmithImportFactoryCreateFileResult UDatasmithSceneElement::ReimportScene()
 {
 	FDatasmithImportFactoryCreateFileResult Result;
 
-	if (this == nullptr || !ImportContextPtr.IsValid() || !ImportContextPtr->Options.IsValid() || !ExternalSourcePtr.IsValid())
+	if (!ImportContextPtr.IsValid() || ImportContextPtr->Options.IsNull() || !ExternalSourcePtr.IsValid())
 	{
 		UE_LOG(LogDatasmithImport, Error, TEXT("Invalid State. Ensure GetExistingDatasmithScene has been called."));
 		return Result;
@@ -824,7 +818,7 @@ UObject* UDatasmithSceneElement::GetOptions(UClass* OptionType)
 	if (ImportContextPtr.IsValid())
 	{
 		// Standard options from Datasmith
-		if (ImportContextPtr->Options.IsValid() && ImportContextPtr->Options->GetClass()->IsChildOf(OptionType))
+		if (!ImportContextPtr->Options.IsNull() && ImportContextPtr->Options->GetClass()->IsChildOf(OptionType))
 		{
 			return ImportContextPtr->Options.Get();
 		}
@@ -843,32 +837,32 @@ UObject* UDatasmithSceneElement::GetOptions(UClass* OptionType)
 
 TMap<UClass*, UObject*> UDatasmithSceneElement::GetAllOptions()
 {
-	TMap<UClass*, UObject*> M;
+	TMap<UClass*, UObject*> AllOptions;
 
 	auto Append = [&](UObject* Option)
 	{
 		if (Option)
-	{
-			M.Add(Option->GetClass(), Option);
-	}
+		{
+			AllOptions.Add(Option->GetClass(), Option);
+		}
 	};
 
 	if (ImportContextPtr.IsValid())
 	{
 		// Standard options from Datasmith
-		if (ImportContextPtr->Options.IsValid())
+		if (ImportContextPtr->Options)
 		{
-			Append(ImportContextPtr->Options.Get());
+			Append(ImportContextPtr->Options);
 		}
 
 		// Additional options from specific translators
 		for(const auto& AdditionalOption : ImportContextPtr->AdditionalImportOptions)
 		{
-			Append(AdditionalOption.Get());
+			Append(AdditionalOption);
 		}
 	}
 
-	return M;
+	return AllOptions;
 }
 
 UDatasmithImportOptions* UDatasmithSceneElement::GetImportOptions()
diff --git a/Engine/Plugins/Enterprise/DatasmithImporter/Source/DatasmithImporter/Private/DatasmithFileProducer.cpp b/Engine/Plugins/Enterprise/DatasmithImporter/Source/DatasmithImporter/Private/DatasmithFileProducer.cpp
--- a/Engine/Plugins/Enterprise/DatasmithImporter/Source/DatasmithImporter/Private/DatasmithFileProducer.cpp
+++ b/Engine/Plugins/Enterprise/DatasmithImporter/Source/DatasmithImporter/Private/DatasmithFileProducer.cpp
@@ -267,9 +267,9 @@ bool UDatasmithFileProducer::InitTranslator()
 	return true;
 }
 
-TArray<TStrongObjectPtr<UDatasmithOptionsBase>> UDatasmithFileProducer::GetTranslatorImportOptions()
+TArray<TObjectPtr<UDatasmithOptionsBase>> UDatasmithFileProducer::GetTranslatorImportOptions()
 {
-	TArray<TStrongObjectPtr<UDatasmithOptionsBase>> Result;
+	TArray<TObjectPtr<UDatasmithOptionsBase>> Result;
 
 	if (!ExternalSourcePtr.IsValid() && !InitTranslator())
 	{
@@ -284,19 +284,19 @@ TArray<TStrongObjectPtr<UDatasmithOptionsBase>> UDatasmithFileProducer::GetTrans
 
 		if (IDatasmithTranslator* Translator = TranslatorPtr.Get())
 		{
-			TArray< TStrongObjectPtr<UDatasmithOptionsBase> > Options;
+			TArray< TObjectPtr<UDatasmithOptionsBase> > Options;
 			Translator->GetSceneImportOptions(Options);
 
-			for (TStrongObjectPtr<UDatasmithOptionsBase> Option : Options)
+			for (TObjectPtr<UDatasmithOptionsBase> Option : Options)
 			{
-				UDatasmithOptionsBase* MyOption = DuplicateObject<UDatasmithOptionsBase>(Option.Get(), this);
+				UDatasmithOptionsBase* MyOption = DuplicateObject<UDatasmithOptionsBase>(Option, this);
 				if (UDatasmithCommonTessellationOptions* TesselationOption = Cast<UDatasmithCommonTessellationOptions>(MyOption))
 				{
 					TesselationOption->Options = DefaultTessellationOptions;
 					TranslatorImportOptions.Add(MyOption);
 					continue;
 				}
-				else if (UDatasmithImportOptions* DatasmithOptions = Cast<UDatasmithImportOptions>(Option.Get()))
+				else if (UDatasmithImportOptions* DatasmithOptions = Cast<UDatasmithImportOptions>(Option))
 				{
 					continue;
 				}
@@ -318,7 +318,7 @@ TArray<TStrongObjectPtr<UDatasmithOptionsBase>> UDatasmithFileProducer::GetTrans
 
 	for (UDatasmithOptionsBase* Option : TranslatorImportOptions)
 	{
-		Result.Add(TStrongObjectPtr<UDatasmithOptionsBase>(Option));
+		Result.Add(Option);
 	}
 
 	return MoveTemp(Result);
@@ -805,7 +805,7 @@ void UDatasmithFileProducer::OnChangeImportSettings()
 		.Title(FText::Format(LOCTEXT("ImportSettingsTitle", "{0} Import Settings"), FText::FromName(TranslatorPtr->GetFName())))
 		.SizingRule(ESizingRule::Autosized);
 
-	TArray<TStrongObjectPtr<UDatasmithOptionsBase>> Options = GetTranslatorImportOptions();
+	TArray<TObjectPtr<UDatasmithOptionsBase>> Options = GetTranslatorImportOptions();
 
 	if (Options.Num() == 0)
 	{
@@ -814,7 +814,7 @@ void UDatasmithFileProducer::OnChangeImportSettings()
 	}
 
 	TArray<UObject*> OptionsRaw;
-	for(TStrongObjectPtr<UDatasmithOptionsBase>& ObjectPtr : Options)
+	for(TObjectPtr<UDatasmithOptionsBase>& ObjectPtr : Options)
 	{
 		OptionsRaw.Add(ObjectPtr.Get());
 	}
@@ -1041,10 +1041,10 @@ void UDatasmithDirProducer::OnChangeImportSettings()
 void UDatasmithDirProducer::SetFileProducerSettings()
 {
 	// Set translator options
-	TArray< TStrongObjectPtr<UDatasmithOptionsBase> > Options = FileProducer->GetTranslatorImportOptions();
+	TArray< TObjectPtr<UDatasmithOptionsBase> > Options = FileProducer->GetTranslatorImportOptions();
 
 	bool bHasTesselationOptions = false;
-	for (TStrongObjectPtr<UDatasmithOptionsBase> Option : Options)
+	for (TObjectPtr<UDatasmithOptionsBase> Option : Options)
 	{
 		if (UDatasmithCommonTessellationOptions* TesselationOption = Cast<UDatasmithCommonTessellationOptions>(Option.Get()))
 		{
@@ -1583,7 +1583,7 @@ void FDatasmithFileProducerDetails::CustomizeDetails(IDetailLayoutBuilder& Detai
 
 	TSharedPtr<STextBlock> IconText;
 
-	TArray<TStrongObjectPtr<UDatasmithOptionsBase>> Options = FileProducer->GetTranslatorImportOptions();
+	TArray<TObjectPtr<UDatasmithOptionsBase>> Options = FileProducer->GetTranslatorImportOptions();
 
 	CustomAssetImportRow.NameContent()
 	[
diff --git a/Engine/Plugins/Enterprise/DatasmithImporter/Source/DatasmithImporter/Private/DatasmithFileProducer.h b/Engine/Plugins/Enterprise/DatasmithImporter/Source/DatasmithImporter/Private/DatasmithFileProducer.h
--- a/Engine/Plugins/Enterprise/DatasmithImporter/Source/DatasmithImporter/Private/DatasmithFileProducer.h
+++ b/Engine/Plugins/Enterprise/DatasmithImporter/Source/DatasmithImporter/Private/DatasmithFileProducer.h
@@ -51,9 +51,9 @@ public:
 
 	void OnChangeImportSettings();
 
-	TArray<TStrongObjectPtr<UDatasmithOptionsBase>> GetTranslatorImportOptions();
+	TArray<TObjectPtr<UDatasmithOptionsBase>> GetTranslatorImportOptions();
 
-	void SetTranslatorImportOptions(TArray<TStrongObjectPtr<UDatasmithOptionsBase>> InOptions)
+	void SetTranslatorImportOptions(TArray<TObjectPtr<UDatasmithOptionsBase>> InOptions)
 	{
 		ImportOptionsOverride = InOptions;
 	}
@@ -96,7 +96,7 @@ private:
 	TUniquePtr< FDataprepWorkReporter > ProgressTaskPtr;
 	TSharedPtr< UE::DatasmithImporter::FExternalSource > ExternalSourcePtr;
 
-	TArray<TStrongObjectPtr<UDatasmithOptionsBase>> ImportOptionsOverride;
+	TArray<TObjectPtr<UDatasmithOptionsBase>> ImportOptionsOverride;
 
 	UPROPERTY( Transient, DuplicateTransient )
 	TObjectPtr<UDatasmithScene> DatasmithScene;
diff --git a/Engine/Plugins/Enterprise/DatasmithImporter/Source/DatasmithImporter/Private/DatasmithImportContext.cpp b/Engine/Plugins/Enterprise/DatasmithImporter/Source/DatasmithImporter/Private/DatasmithImportContext.cpp
--- a/Engine/Plugins/Enterprise/DatasmithImporter/Source/DatasmithImporter/Private/DatasmithImportContext.cpp
+++ b/Engine/Plugins/Enterprise/DatasmithImporter/Source/DatasmithImporter/Private/DatasmithImportContext.cpp
@@ -259,11 +259,11 @@ void FDatasmithImportContext::UpdateImportOption(UDatasmithOptionsBase* NewOptio
 {
 	if (NewOption)
 	{
-		auto SameClass = [=](const TStrongObjectPtr<UDatasmithOptionsBase>& Option) { return Option->GetClass() == NewOption->GetClass(); };
+		auto SameClass = [=](const TObjectPtr<UDatasmithOptionsBase>& Option) { return Option->GetClass() == NewOption->GetClass(); };
 
-		if (TStrongObjectPtr<UDatasmithOptionsBase>* PreviousOptionPtr = AdditionalImportOptions.FindByPredicate(SameClass))
+		if (TObjectPtr<UDatasmithOptionsBase>* PreviousOptionPtr = AdditionalImportOptions.FindByPredicate(SameClass))
 		{
-			PreviousOptionPtr->Reset(NewOption);
+			*PreviousOptionPtr = NewOption;
 		}
 		else
 		{
@@ -320,9 +320,9 @@ bool FDatasmithImportContext::InitOptions(const TSharedPtr<FJsonObject>& ImportS
 	ImportOptions.Reserve(1 + AdditionalImportOptions.Num());
 
 	ImportOptions.Add(Options.Get());
-	for (const TStrongObjectPtr<UDatasmithOptionsBase>& AdditionalOption : AdditionalImportOptions)
+	for (const TObjectPtr<UDatasmithOptionsBase>& AdditionalOption : AdditionalImportOptions)
 	{
-		ImportOptions.Add(AdditionalOption.Get());
+		ImportOptions.Add(AdditionalOption);
 	}
 
 	if (bSilent)
@@ -581,6 +581,12 @@ void FDatasmithImportContext::FInternalReferenceCollector::AddReferencedObjects(
 
 	DatasmithImportContextInternal::AddReferenceList(Collector, ImportContext->ImportedActorMap);
 	DatasmithImportContextInternal::AddReferenceList(Collector, ImportContext->ImportedSceneComponentMap);
+
+	Collector.AddReferencedObject(ImportContext->Options);
+	for (TObjectPtr<UDatasmithOptionsBase>& Option : ImportContext->AdditionalImportOptions)
+	{
+		Collector.AddReferencedObject(Option);
+	}
 }
 
 FDatasmithActorUniqueLabelProvider::FDatasmithActorUniqueLabelProvider(UWorld* World)
diff --git a/Engine/Plugins/Enterprise/DatasmithImporter/Source/DatasmithImporter/Private/DatasmithImportFactory.cpp b/Engine/Plugins/Enterprise/DatasmithImporter/Source/DatasmithImporter/Private/DatasmithImportFactory.cpp
--- a/Engine/Plugins/Enterprise/DatasmithImporter/Source/DatasmithImporter/Private/DatasmithImportFactory.cpp
+++ b/Engine/Plugins/Enterprise/DatasmithImporter/Source/DatasmithImporter/Private/DatasmithImportFactory.cpp
@@ -161,7 +161,7 @@ namespace DatasmithImportFactoryImpl
 		// Copy over the changes the user may have done on the options
 		ReImportSceneData->BaseOptions = InContext.Options->BaseOptions;
 
-		for (const TStrongObjectPtr<UDatasmithOptionsBase>& Option : InContext.AdditionalImportOptions)
+		for (const TObjectPtr<UDatasmithOptionsBase>& Option : InContext.AdditionalImportOptions)
 		{
 			UDatasmithOptionsBase* OptionObj = Option.Get();
 			OptionObj->Rename(nullptr, ReImportSceneData);
@@ -186,7 +186,7 @@ namespace DatasmithImportFactoryImpl
 		bOutOperationCancelled = false;
 
 		// Return if the context is not valid
-		if ( !InContext.Options.IsValid() )
+		if ( InContext.Options.IsNull() )
 		{
 			return false;
 		}
@@ -348,12 +348,12 @@ namespace DatasmithImportFactoryImpl
 
 				// Log tessellator if CADKernel has been used
 				bool bUseCADKernel = false;
-				TArray<TStrongObjectPtr<UDatasmithOptionsBase>> Options;
+				TArray<TObjectPtr<UDatasmithOptionsBase>> Options;
 				ImportContext.SceneTranslator->GetSceneImportOptions(Options);
 
-				for (const TStrongObjectPtr<UDatasmithOptionsBase>& Option : Options)
+				for (const TObjectPtr<UDatasmithOptionsBase>& Option : Options)
 				{
-					if (UDatasmithCommonTessellationOptions* TessellationOptionsObject = Cast<UDatasmithCommonTessellationOptions>(Option.Get()))
+					if (UDatasmithCommonTessellationOptions* TessellationOptionsObject = Cast<UDatasmithCommonTessellationOptions>(Option))
 					{
 						bUseCADKernel = TessellationOptionsObject->Options.bUseCADKernel;
 					}
diff --git a/Engine/Plugins/Enterprise/DatasmithImporter/Source/DatasmithImporter/Public/DatasmithImportContext.h b/Engine/Plugins/Enterprise/DatasmithImporter/Source/DatasmithImporter/Public/DatasmithImportContext.h
--- a/Engine/Plugins/Enterprise/DatasmithImporter/Source/DatasmithImporter/Public/DatasmithImportContext.h
+++ b/Engine/Plugins/Enterprise/DatasmithImporter/Source/DatasmithImporter/Public/DatasmithImportContext.h
@@ -228,10 +228,10 @@ struct DATASMITHIMPORTER_API FDatasmithImportContext
 	EObjectFlags ObjectFlags;
 
 	/** Settings related to the import of a Datasmith scene */
-	TStrongObjectPtr<UDatasmithImportOptions> Options;
+	TObjectPtr<UDatasmithImportOptions> Options;
 
 	/** Per-Translator settings related to the import of a Datasmith scene */
-	TArray<TStrongObjectPtr<UDatasmithOptionsBase>> AdditionalImportOptions;
+	TArray<TObjectPtr<UDatasmithOptionsBase>> AdditionalImportOptions;
 
 	/** Stack of parent components we're currently importing under */
 	TArray< USceneComponent* > Hierarchy;
diff --git a/Engine/Plugins/Enterprise/DatasmithImporter/Source/DatasmithTranslator/Public/DatasmithTranslator.h b/Engine/Plugins/Enterprise/DatasmithImporter/Source/DatasmithTranslator/Public/DatasmithTranslator.h
--- a/Engine/Plugins/Enterprise/DatasmithImporter/Source/DatasmithTranslator/Public/DatasmithTranslator.h
+++ b/Engine/Plugins/Enterprise/DatasmithImporter/Source/DatasmithTranslator/Public/DatasmithTranslator.h
@@ -121,6 +121,8 @@ public:
 	 *
 	 * @param OptionClasses list of classes that will be displayed to the user
 	 */
+	virtual void GetSceneImportOptions(TArray<TObjectPtr<UDatasmithOptionsBase>>& Options) {}
+	UE_DEPRECATED(5.2, "Deprecated, please use same method using array of TObjectPtr instead")
 	virtual void GetSceneImportOptions(TArray<TStrongObjectPtr<UDatasmithOptionsBase>>& Options) {}
 
 	/**
@@ -128,6 +130,8 @@ public:
 	 *
 	 * @param Options Actual values for the displayed options.
 	 */
+	virtual void SetSceneImportOptions(const TArray<TObjectPtr<UDatasmithOptionsBase>>& Options) {}
+	UE_DEPRECATED(5.2, "Deprecated, please use same method using array of TObjectPtr instead")
 	virtual void SetSceneImportOptions(TArray<TStrongObjectPtr<UDatasmithOptionsBase>>& Options) {}
 
 private:
@@ -181,6 +185,12 @@ namespace Datasmith
 		return Option;
 	}
 
+	template<class UOptionClass>
+	inline TObjectPtr<UOptionClass> MakeOptionsObjectPtr()
+	{
+		return NewObject<UOptionClass>(GetTransientPackage(), UOptionClass::StaticClass());
+	}
+
 	FString DATASMITHTRANSLATOR_API GetXMLFileSchema(const FString& XmlFilePath);
 
 	template<typename ValueType>
diff --git a/Engine/Plugins/VirtualProduction/DMX/DatasmithMVR/Source/DatasmithMVRTranslator/Private/DatasmithMVRNativeTranslator.cpp b/Engine/Plugins/VirtualProduction/DMX/DatasmithMVR/Source/DatasmithMVRTranslator/Private/DatasmithMVRNativeTranslator.cpp
--- a/Engine/Plugins/VirtualProduction/DMX/DatasmithMVR/Source/DatasmithMVRTranslator/Private/DatasmithMVRNativeTranslator.cpp
+++ b/Engine/Plugins/VirtualProduction/DMX/DatasmithMVR/Source/DatasmithMVRTranslator/Private/DatasmithMVRNativeTranslator.cpp
@@ -55,7 +55,7 @@ bool FDatasmithMVRNativeTranslator::LoadScene(TSharedRef<IDatasmithScene> InOutS
 	return true;
 }
 
-void FDatasmithMVRNativeTranslator::GetSceneImportOptions(TArray<TStrongObjectPtr<UDatasmithOptionsBase>>& Options)
+void FDatasmithMVRNativeTranslator::GetSceneImportOptions(TArray<TObjectPtr<UDatasmithOptionsBase>>& Options)
 {
 	FDatasmithNativeTranslator::GetSceneImportOptions(Options);
 
@@ -66,16 +66,16 @@ void FDatasmithMVRNativeTranslator::GetSceneImportOptions(TArray<TStrongObjectPt
 		ImportOptions = Datasmith::MakeOptions<UDatasmithMVRImportOptions>();
 	}
 
-	Options.Add(ImportOptions);
+	Options.Add(ImportOptions.Get());
 }
 
-void FDatasmithMVRNativeTranslator::SetSceneImportOptions(TArray<TStrongObjectPtr<UDatasmithOptionsBase>>& Options)
+void FDatasmithMVRNativeTranslator::SetSceneImportOptions(const TArray<TObjectPtr<UDatasmithOptionsBase>>& Options)
 {
 	FDatasmithNativeTranslator::SetSceneImportOptions(Options);
 
-	for (const TStrongObjectPtr<UDatasmithOptionsBase>& OptionPtr : Options)
+	for (const TObjectPtr<UDatasmithOptionsBase>& OptionPtr : Options)
 	{
-		if (UDatasmithMVRImportOptions* InImportOptions = Cast<UDatasmithMVRImportOptions>(OptionPtr.Get()))
+		if (UDatasmithMVRImportOptions* InImportOptions = Cast<UDatasmithMVRImportOptions>(OptionPtr))
 		{
 			ImportOptions.Reset(InImportOptions);
 		}
diff --git a/Engine/Plugins/VirtualProduction/DMX/DatasmithMVR/Source/DatasmithMVRTranslator/Private/DatasmithMVRNativeTranslator.h b/Engine/Plugins/VirtualProduction/DMX/DatasmithMVR/Source/DatasmithMVRTranslator/Private/DatasmithMVRNativeTranslator.h
--- a/Engine/Plugins/VirtualProduction/DMX/DatasmithMVR/Source/DatasmithMVRTranslator/Private/DatasmithMVRNativeTranslator.h
+++ b/Engine/Plugins/VirtualProduction/DMX/DatasmithMVR/Source/DatasmithMVRTranslator/Private/DatasmithMVRNativeTranslator.h
@@ -18,8 +18,8 @@ public:
 	//~ Begin DatasmithNativeTranslator interface
 	virtual FName GetFName() const override { return "DatasmithMVRNativeTranslator"; };
 	virtual bool LoadScene(TSharedRef<IDatasmithScene> InOutScene) override;
-	virtual void GetSceneImportOptions(TArray<TStrongObjectPtr<UDatasmithOptionsBase>>& Options) override;
-	virtual void SetSceneImportOptions(TArray<TStrongObjectPtr<UDatasmithOptionsBase>>& Options) override;
+	virtual void GetSceneImportOptions(TArray<TObjectPtr<UDatasmithOptionsBase>>& Options) override;
+	virtual void SetSceneImportOptions(const TArray<TObjectPtr<UDatasmithOptionsBase>>& Options) override;
 	//~ End DatasmithNativeTranslator interface
 
 private:
