FNames containing asset paths are deprecated. FSoftObjectPath or FTopLevelAssetPath should be used instead.
Fixups for Datasmith/Dataprep.

diff --git a/Engine/Plugins/Enterprise/DataprepEditor/Source/DataprepCore/Private/DataprepCoreUtils.cpp b/Engine/Plugins/Enterprise/DataprepEditor/Source/DataprepCore/Private/DataprepCoreUtils.cpp
--- a/Engine/Plugins/Enterprise/DataprepEditor/Source/DataprepCore/Private/DataprepCoreUtils.cpp
+++ b/Engine/Plugins/Enterprise/DataprepEditor/Source/DataprepCore/Private/DataprepCoreUtils.cpp
@@ -751,7 +751,7 @@ void FDataprepCoreUtils::DeleteTemporaryFolders(const FString& BaseTemporaryPath
 		ObjectsToDelete.Reserve( AssetDataList.Num() );
 		for(const FAssetData& AssetData : AssetDataList)
 		{
-			FSoftObjectPath ObjectPath( AssetData.ObjectPath.ToString() );
+			FSoftObjectPath ObjectPath( AssetData.GetSoftObjectPath() );
 
 			if(UObject* Object = ObjectPath.ResolveObject())
 			{
diff --git a/Engine/Plugins/Enterprise/DataprepEditor/Source/DataprepLibraries/Private/DataprepEditingOperations.cpp b/Engine/Plugins/Enterprise/DataprepEditor/Source/DataprepLibraries/Private/DataprepEditingOperations.cpp
--- a/Engine/Plugins/Enterprise/DataprepEditor/Source/DataprepLibraries/Private/DataprepEditingOperations.cpp
+++ b/Engine/Plugins/Enterprise/DataprepEditor/Source/DataprepLibraries/Private/DataprepEditingOperations.cpp
@@ -715,7 +715,7 @@ const FAssetData& FDataprepSpawnActorsAtLocationDetails::GetAssetData() const
 {
 	if (DataprepOperation->SelectedAsset)
 	{
-		if (!DataprepOperation->SelectedAsset->GetPathName().Equals(CachedAssetData.ObjectPath.ToString(), ESearchCase::CaseSensitive))
+		if (FSoftObjectPath(DataprepOperation->SelectedAsset) == CachedAssetData.GetSoftObjectPath())
 		{
 			// This always uses the exact object pointed at
 			CachedAssetData = FAssetData(DataprepOperation->SelectedAsset, true);
diff --git a/Engine/Plugins/Enterprise/DatasmithImporter/Source/DatasmithImporter/Private/DatasmithImporter.cpp b/Engine/Plugins/Enterprise/DatasmithImporter/Source/DatasmithImporter/Private/DatasmithImporter.cpp
--- a/Engine/Plugins/Enterprise/DatasmithImporter/Source/DatasmithImporter/Private/DatasmithImporter.cpp
+++ b/Engine/Plugins/Enterprise/DatasmithImporter/Source/DatasmithImporter/Private/DatasmithImporter.cpp
@@ -1761,7 +1761,7 @@ void FDatasmithImporter::FilterElementsToImport( FDatasmithImportContext& Import
 
 	auto ElementNeedsReimport = [&](const FString& FullyQualifiedName, TSharedRef<IDatasmithElement> Element, const FString& SourcePath) -> bool
 	{
-		const FAssetData AssetData = AssetRegistry.GetAssetByObjectPath(*FullyQualifiedName);
+		const FAssetData AssetData = AssetRegistry.GetAssetByObjectPath(FSoftObjectPath(FullyQualifiedName));
 		const FAssetDataTagMapSharedView::FFindTagResult ImportDataStr = AssetData.TagsAndValues.FindTag(TEXT("AssetImportData"));
 		FString CurrentRelativeFileName;
 
@@ -1820,7 +1820,7 @@ void FDatasmithImporter::FilterElementsToImport( FDatasmithImportContext& Import
 		// If the mesh element does not need to be re-imported, register its name
 		else
 		{
-			const FAssetData AssetData = AssetRegistry.GetAssetByObjectPath( *AssetName);
+			const FAssetData AssetData = AssetRegistry.GetAssetByObjectPath( FSoftObjectPath(AssetName) );
 			ImportContext.AssetsContext.StaticMeshNameProvider.AddExistingName( FPaths::GetBaseFilename( AssetData.PackageName.ToString() ) );
 		}
 	}
@@ -1847,7 +1847,7 @@ void FDatasmithImporter::FilterElementsToImport( FDatasmithImportContext& Import
 		// If the texture element does not need to be re-imported, register its name
 		else
 		{
-			const FAssetData AssetData = AssetRegistry.GetAssetByObjectPath( *AssetName);
+			const FAssetData AssetData = AssetRegistry.GetAssetByObjectPath( FSoftObjectPath(AssetName) );
 			ImportContext.AssetsContext.TextureNameProvider.AddExistingName( FPaths::GetBaseFilename( AssetData.PackageName.ToString() ) );
 		}
 	}
diff --git a/Engine/Plugins/Enterprise/DatasmithImporter/Source/DatasmithImporter/Private/Utility/DatasmithImporterUtils.cpp b/Engine/Plugins/Enterprise/DatasmithImporter/Source/DatasmithImporter/Private/Utility/DatasmithImporterUtils.cpp
--- a/Engine/Plugins/Enterprise/DatasmithImporter/Source/DatasmithImporter/Private/Utility/DatasmithImporterUtils.cpp
+++ b/Engine/Plugins/Enterprise/DatasmithImporter/Source/DatasmithImporter/Private/Utility/DatasmithImporterUtils.cpp
@@ -512,7 +512,7 @@ bool FDatasmithImporterUtils::CanCreateAsset(const FString& AssetPathName, const
 		case EAssetCreationStatus::CS_ClassMismatch:
 		{
 			IAssetRegistry& AssetRegistry = FModuleManager::LoadModuleChecked<FAssetRegistryModule>(TEXT("AssetRegistry")).Get();
-			const FAssetData AssetData = AssetRegistry.GetAssetByObjectPath(*AssetPathName);
+			const FAssetData AssetData = AssetRegistry.GetAssetByObjectPath(FSoftObjectPath(AssetPathName));
 
 			const FString FoundClassName(AssetData.GetClass()->GetFName().ToString());
 			const FString ExpectedClassName(AssetClass->GetFName().ToString());
@@ -547,7 +547,7 @@ FDatasmithImporterUtils::EAssetCreationStatus FDatasmithImporterUtils::CanCreate
 
 	IAssetRegistry& AssetRegistry = FModuleManager::LoadModuleChecked<FAssetRegistryModule>(TEXT("AssetRegistry")).Get();
 
-	const FAssetData AssetData = AssetRegistry.GetAssetByObjectPath(*AssetPathName);
+	const FAssetData AssetData = AssetRegistry.GetAssetByObjectPath(FSoftObjectPath(AssetPathName));
 
 	// Asset does not exist yet. Safe to import
 	if (!AssetData.IsValid())
