Made the InterchangeEditor and DataprepEditor dependencies editor only in the DatasmithImporter.

diff --git a/Engine/Plugins/Enterprise/DatasmithImporter/DatasmithImporter.uplugin b/Engine/Plugins/Enterprise/DatasmithImporter/DatasmithImporter.uplugin
--- a/Engine/Plugins/Enterprise/DatasmithImporter/DatasmithImporter.uplugin
+++ b/Engine/Plugins/Enterprise/DatasmithImporter/DatasmithImporter.uplugin
@@ -67,11 +67,17 @@
 		},
 		{
 			"Name": "DataPrepEditor",
-			"Enabled": true
+			"Enabled": true,
+			"TargetAllowList": [
+				"Editor"
+			]
 		},
 		{
 			"Name": "InterchangeEditor",
-			"Enabled": true
+			"Enabled": true,
+			"TargetAllowList": [
+				"Editor"
+			]
 		},
 		{
 			"Name": "VariantManagerContent",
