Datasmith: Cloth import : fix crash on reimport

diff --git a/Engine/Plugins/Enterprise/DatasmithImporter/Source/DatasmithImporter/Private/DatasmithImporter.cpp b/Engine/Plugins/Enterprise/DatasmithImporter/Source/DatasmithImporter/Private/DatasmithImporter.cpp
--- a/Engine/Plugins/Enterprise/DatasmithImporter/Source/DatasmithImporter/Private/DatasmithImporter.cpp
+++ b/Engine/Plugins/Enterprise/DatasmithImporter/Source/DatasmithImporter/Private/DatasmithImporter.cpp
@@ -576,20 +576,16 @@ UStaticMesh* FDatasmithImporter::FinalizeStaticMesh( UStaticMesh* SourceStaticMe
 
 UObject* FDatasmithImporter::FinalizeCloth(UObject* SourceCloth, const TCHAR* FolderPath, UObject* ExistingCloth, TMap<UObject*, UObject*>* ReferencesToRemap)
 {
-	// #ue_ds_cloth_note FDatasmithImporter::FinalizeCloth
 	if (Cast<USkinnedAsset>(ExistingCloth))
 	{
-		// #ue_ds_cloth_todo review if this step is legit
-		FSkinnedAssetCompilingManager::Get().FinishCompilation({Cast<USkinnedAsset>(ExistingCloth)});
+		// Delete the old asset completely. This fixes a crash with the render ressources being used while the asset is replaced.
+		TArray<UObject*> ExistingClothAsArray{ExistingCloth};
+		ObjectTools::ForceReplaceReferences(SourceCloth, ExistingClothAsArray);
+		ObjectTools::DeleteAssets({ExistingCloth}, false);
+		ExistingCloth = nullptr;
 	}
-
 	UObject* Asset = FDatasmithImporterImpl::FinalizeAsset(SourceCloth, FolderPath, ExistingCloth, ReferencesToRemap);
 
-	if (Cast<USkinnedAsset>(Asset))
-	{
-		// #ue_ds_cloth_todo review if this step is legit
-		FSkinnedAssetCompilingManager::Get().FinishCompilation({Cast<USkinnedAsset>(Asset)});
-	}
 	return Asset;
 }
 
