compile fixes for deprecated TObjectPtr.IsNull

diff --git a/Engine/Plugins/Enterprise/DatasmithImporter/Source/DatasmithImporter/Private/DatasmithBlueprintLibrary.cpp b/Engine/Plugins/Enterprise/DatasmithImporter/Source/DatasmithImporter/Private/DatasmithBlueprintLibrary.cpp
--- a/Engine/Plugins/Enterprise/DatasmithImporter/Source/DatasmithImporter/Private/DatasmithBlueprintLibrary.cpp
+++ b/Engine/Plugins/Enterprise/DatasmithImporter/Source/DatasmithImporter/Private/DatasmithBlueprintLibrary.cpp
@@ -743,7 +743,7 @@ FDatasmithImportFactoryCreateFileResult UDatasmithSceneElement::ReimportScene()
 {
 	FDatasmithImportFactoryCreateFileResult Result;
 
-	if (!ImportContextPtr.IsValid() || ImportContextPtr->Options.IsNull() || !ExternalSourcePtr.IsValid())
+	if (!ImportContextPtr.IsValid() || !ImportContextPtr->Options || !ExternalSourcePtr.IsValid())
 	{
 		UE_LOG(LogDatasmithImport, Error, TEXT("Invalid State. Ensure GetExistingDatasmithScene has been called."));
 		return Result;
@@ -818,7 +818,7 @@ UObject* UDatasmithSceneElement::GetOptions(UClass* OptionType)
 	if (ImportContextPtr.IsValid())
 	{
 		// Standard options from Datasmith
-		if (!ImportContextPtr->Options.IsNull() && ImportContextPtr->Options->GetClass()->IsChildOf(OptionType))
+		if (ImportContextPtr->Options && ImportContextPtr->Options->GetClass()->IsChildOf(OptionType))
 		{
 			return ImportContextPtr->Options.Get();
 		}
diff --git a/Engine/Plugins/Enterprise/DatasmithImporter/Source/DatasmithImporter/Private/DatasmithImportFactory.cpp b/Engine/Plugins/Enterprise/DatasmithImporter/Source/DatasmithImporter/Private/DatasmithImportFactory.cpp
--- a/Engine/Plugins/Enterprise/DatasmithImporter/Source/DatasmithImporter/Private/DatasmithImportFactory.cpp
+++ b/Engine/Plugins/Enterprise/DatasmithImporter/Source/DatasmithImporter/Private/DatasmithImportFactory.cpp
@@ -186,7 +186,7 @@ namespace DatasmithImportFactoryImpl
 		bOutOperationCancelled = false;
 
 		// Return if the context is not valid
-		if ( InContext.Options.IsNull() )
+		if ( !InContext.Options )
 		{
 			return false;
 		}
diff --git a/Engine/Plugins/Experimental/Animation/PoseSearch/Source/Runtime/Private/PoseSearch.cpp b/Engine/Plugins/Experimental/Animation/PoseSearch/Source/Runtime/Private/PoseSearch.cpp
--- a/Engine/Plugins/Experimental/Animation/PoseSearch/Source/Runtime/Private/PoseSearch.cpp
+++ b/Engine/Plugins/Experimental/Animation/PoseSearch/Source/Runtime/Private/PoseSearch.cpp
@@ -370,7 +370,7 @@ void UPoseSearchSchema::Finalize()
 	using namespace UE::PoseSearch;
 
 	// Discard null channels
-	Channels.RemoveAll([](TObjectPtr<UPoseSearchFeatureChannel>& Channel) { return Channel.IsNull(); });
+	Channels.RemoveAll([](TObjectPtr<UPoseSearchFeatureChannel>& Channel) { return !Channel; });
 
 	BoneReferences.Reset();
 
diff --git a/Engine/Plugins/Experimental/Enterprise/DatasmithInterchange/Source/DatasmithInterchange/Private/InterchangeDatasmithPipeline.cpp b/Engine/Plugins/Experimental/Enterprise/DatasmithInterchange/Source/DatasmithInterchange/Private/InterchangeDatasmithPipeline.cpp
--- a/Engine/Plugins/Experimental/Enterprise/DatasmithInterchange/Source/DatasmithInterchange/Private/InterchangeDatasmithPipeline.cpp
+++ b/Engine/Plugins/Experimental/Enterprise/DatasmithInterchange/Source/DatasmithInterchange/Private/InterchangeDatasmithPipeline.cpp
@@ -76,7 +76,7 @@ void UInterchangeDatasmithPipeline::ExecutePreImportPipeline(UInterchangeBaseNod
 		}
 	};
 
-	ensure(!Results.IsNull());
+	ensure(Results);
 	ExecutePreImportPipelineFunc(TexturePipeline);
 	ExecutePreImportPipelineFunc(MaterialPipeline);
 	ExecutePreImportPipelineFunc(MeshPipeline);
