Datasmith: cloth importer code is now part of the importer plugin

diff --git a/Engine/Plugins/Enterprise/DatasmithImporter/DatasmithImporter.uplugin b/Engine/Plugins/Enterprise/DatasmithImporter/DatasmithImporter.uplugin
--- a/Engine/Plugins/Enterprise/DatasmithImporter/DatasmithImporter.uplugin
+++ b/Engine/Plugins/Enterprise/DatasmithImporter/DatasmithImporter.uplugin
@@ -86,6 +86,10 @@
 		{
 			"Name": "VariantManager",
 			"Enabled": true
+		},
+		{
+			"Name": "ChaosClothAsset",
+			"Enabled": true
 		}
 	]
 }
\ No newline at end of file
diff --git a/Engine/Plugins/Enterprise/DatasmithImporter/Source/DatasmithImporter/DatasmithImporter.Build.cs b/Engine/Plugins/Enterprise/DatasmithImporter/Source/DatasmithImporter/DatasmithImporter.Build.cs
--- a/Engine/Plugins/Enterprise/DatasmithImporter/Source/DatasmithImporter/DatasmithImporter.Build.cs
+++ b/Engine/Plugins/Enterprise/DatasmithImporter/Source/DatasmithImporter/DatasmithImporter.Build.cs
@@ -7,12 +7,14 @@ namespace UnrealBuildTool.Rules
 		public DatasmithImporter(ReadOnlyTargetRules Target)
 			: base(Target)
 		{
-
 			PrivateDependencyModuleNames.AddRange(
 				new string[]
 				{
 					"Analytics",
 					"ApplicationCore",
+					"Chaos",
+					"ChaosClothAsset",
+					"ChaosClothAssetEngine",
 					"CinematicCamera",
 					"ContentBrowser",
 					"Core",
diff --git a/Engine/Plugins/Enterprise/DatasmithImporter/Source/DatasmithImporter/Private/DatasmithActorImporter.cpp b/Engine/Plugins/Enterprise/DatasmithImporter/Source/DatasmithImporter/Private/DatasmithActorImporter.cpp
--- a/Engine/Plugins/Enterprise/DatasmithImporter/Source/DatasmithImporter/Private/DatasmithActorImporter.cpp
+++ b/Engine/Plugins/Enterprise/DatasmithImporter/Source/DatasmithImporter/Private/DatasmithActorImporter.cpp
@@ -3,7 +3,6 @@
 #include "DatasmithActorImporter.h"
 
 #include "DatasmithCameraImporter.h"
-#include "DatasmithClothImporter.h"
 #include "DatasmithImportContext.h"
 #include "DatasmithImporterModule.h"
 #include "DatasmithImportOptions.h"
@@ -23,6 +22,8 @@
 #include "Utility/DatasmithMeshHelper.h"
 
 #include "ActorFactories/ActorFactoryDeferredDecal.h"
+#include "ChaosClothAsset/ClothAsset.h"
+#include "ChaosClothAsset/ClothComponent.h"
 #include "CineCameraActor.h"
 #include "CineCameraComponent.h"
 #include "Components/DecalComponent.h"
@@ -294,7 +295,6 @@ AActor* FDatasmithActorImporter::ImportClothActor(FDatasmithImportContext& Impor
 	if (ImportedActor)
 	{
 		// Find imported resource
-		// #ue_ds_cloth_todo: Can be converted to a FDatasmithImporterUtils::FindAsset call when the actual cloth asset type is known to the importer
 		const TCHAR* ClothName = ClothActorElement->GetCloth();
 		UObject* ClothAsset = nullptr;
 
@@ -318,24 +318,19 @@ AActor* FDatasmithActorImporter::ImportClothActor(FDatasmithImportContext& Impor
 			}
 		}
 
-		if (IDatasmithImporterExt* Ext = IDatasmithImporterModule::Get().GetClothImporterExtension())
+		if (UChaosClothComponent* ClothComponent = NewObject<UChaosClothComponent>(ImportedActor))
 		{
-			if (USceneComponent* ClothComponent = Ext->MakeClothComponent(ImportedActor, ClothAsset))
-			{
-				ImportedActor->AddInstanceComponent(ClothComponent);
-				ClothComponent->SetupAttachment(ImportedActor->GetRootComponent());
+			ClothComponent->SetClothAsset(Cast<UChaosClothAsset>(ClothAsset));
 
-				ClothComponent->bVisualizeComponent = true;
-				ClothComponent->RegisterComponent();
-			}
-			else
-			{
-				ImportContext.LogError(FText::Format(LOCTEXT("ClothComponentCreationFailure", "Cannot create cloth component for asset {0} reauired by actor {1}."), FText::FromString(ClothName), FText::FromString(ClothActorElement->GetLabel())));
-			}
+			ImportedActor->AddInstanceComponent(ClothComponent);
+			ClothComponent->SetupAttachment(ImportedActor->GetRootComponent());
+
+			ClothComponent->bVisualizeComponent = true;
+			ClothComponent->RegisterComponent();
 		}
 		else
 		{
-			ImportContext.LogError(FText::Format(LOCTEXT("ClothImporterExtensionIssue", "Cannot find ClothImporterExtension, cannot import actor {0}."), FText::FromString(ClothActorElement->GetLabel())));
+			ImportContext.LogError(FText::Format(LOCTEXT("ClothComponentCreationFailure", "Cannot create cloth component for asset {0} reauired by actor {1}."), FText::FromString(ClothName), FText::FromString(ClothActorElement->GetLabel())));
 		}
 
 		if (ClothAsset == nullptr)
diff --git a/Engine/Plugins/Enterprise/DatasmithImporter/Source/DatasmithImporter/Private/DatasmithImporter.cpp b/Engine/Plugins/Enterprise/DatasmithImporter/Source/DatasmithImporter/Private/DatasmithImporter.cpp
--- a/Engine/Plugins/Enterprise/DatasmithImporter/Source/DatasmithImporter/Private/DatasmithImporter.cpp
+++ b/Engine/Plugins/Enterprise/DatasmithImporter/Source/DatasmithImporter/Private/DatasmithImporter.cpp
@@ -9,7 +9,6 @@
 #include "DatasmithAssetUserData.h"
 #include "DatasmithCameraImporter.h"
 #include "DatasmithCloth.h"
-#include "DatasmithClothImporter.h"
 #include "DatasmithImportContext.h"
 #include "DatasmithImporterModule.h"
 #include "DatasmithLevelSequenceImporter.h"
@@ -80,6 +79,13 @@
 #include "UObject/UObjectHash.h"
 #include "UnrealEdGlobals.h"
 
+#include "Animation/Skeleton.h"
+#include "ChaosClothAsset/ClothPreset.h"
+#include "ChaosClothAsset/ClothAsset.h"
+#include "ChaosClothAsset/ClothAdapter.h"
+#include "ChaosClothAsset/ClothLodAdapter.h"
+#include "ChaosClothAsset/ClothPatternAdapter.h"
+
 extern UNREALED_API UEditorEngine* GEditor;
 
 #define LOCTEXT_NAMESPACE "DatasmithImporter"
@@ -371,7 +377,6 @@ void FDatasmithImporter::ImportStaticMeshes( FDatasmithImportContext& ImportCont
 
 void FDatasmithImporter::ImportClothes(FDatasmithImportContext& ImportContext)
 {
-	// #ue_ds_cloth_note: FDatasmithImporter::ImportClothes: experimental import code
 	const int32 ClothesCount = ImportContext.FilteredScene->GetClothesCount();
 
 	if (!ImportContext.Options->BaseOptions.bIncludeGeometry
@@ -407,99 +412,123 @@ void FDatasmithImporter::ImportClothes(FDatasmithImportContext& ImportContext)
 
 		FDatasmithImporterImpl::ReportProgress(Progress, 1.f, FText::FromString(FString::Printf(TEXT("Importing cloth %d/%d (%s) ..."), ClothIndex + 1, ClothesCount, ClothElement->GetLabel())));
 
-		UObject* ExistingCloth = nullptr;
-		if (ImportContext.SceneAsset)
+		int32 MaxNameCharCount = FDatasmithImporterUtils::GetAssetNameMaxCharCount(ImportContext.AssetsContext.StaticMeshesFinalPackage.Get());
+		FString ClothName = ImportContext.AssetsContext.StaticMeshNameProvider.GenerateUniqueName(ClothElement->GetLabel(), MaxNameCharCount);
+		ClothName = ObjectTools::SanitizeObjectName(ClothName);
+
+		FDatasmithClothElementPayload Payload;
+		// #ue_ds_cloth_todo: async parallel support
 		{
-			if (TSoftObjectPtr<UObject>* ExistingAssetPtr = ImportContext.SceneAsset->Clothes.Find(ClothElement->GetName()))
+			TRACE_CPUPROFILER_EVENT_SCOPE(FDatasmithImporter::ImportClothes:TranslatorLoadCloth);
+			if (!ImportContext.SceneTranslator->LoadCloth(ClothElement, Payload))
 			{
-				ExistingCloth = ExistingAssetPtr->LoadSynchronous();
+				UE_LOG(LogDatasmithImport, Warning, TEXT("Cloth %s: load failed, translator issue."), ClothElement->GetName());
+				return;
 			}
 		}
+		FDatasmithCloth DsCloth = MoveTemp(Payload.Cloth);
 
-		{ // copy of FDatasmithImporter::ImportStaticMesh
+		UPackage* ImportPackage = ImportContext.AssetsContext.StaticMeshesImportPackage.Get();
+		EObjectFlags ObjectFlags = ImportContext.ObjectFlags & ~RF_Public; // not RF_Public yet, the publicized asset will be.
 
-			{ // copy of FDatasmithStaticMeshImporter::ImportStaticMesh
-// 				// 2. get destination asset
-// 				// 2.1. FDatasmithImporterUtils::CanCreateAsset<UStaticMesh>
-				int32 MaxNameCharCount = FDatasmithImporterUtils::GetAssetNameMaxCharCount(ImportContext.AssetsContext.StaticMeshesFinalPackage.Get());
-				FString ClothName = ImportContext.AssetsContext.StaticMeshNameProvider.GenerateUniqueName(ClothElement->GetLabel(), MaxNameCharCount);
-				ClothName = ObjectTools::SanitizeObjectName(ClothName);
+		UChaosClothAsset* ClothAsset = nullptr;
+		TArray<UObject*> ClothPresetAssets;
 
-				UPackage* Outer = ImportContext.AssetsContext.StaticMeshesImportPackage.Get();
-				EObjectFlags ObjectFlags = ImportContext.ObjectFlags & ~RF_Public; // not RF_Public yet, the publicized asset will be.
-
-				FDatasmithClothElementPayload Payload;
-				// #ue_ds_cloth_todo: async parallel support
+		for (const FDatasmithClothPresetPropertySet& PropertySet : DsCloth.PropertySets)
+		{
+			FString PresetAssetName = ObjectTools::SanitizeObjectName(PropertySet.SetName);
+			if (UChaosClothPreset* PropertySetAsset = NewObject<UChaosClothPreset>(ImportPackage, *PresetAssetName, ObjectFlags))
+			{
+				ClothPresetAssets.Add(PropertySetAsset);
+				for (const FDatasmithClothPresetProperty& Property : PropertySet.Properties)
 				{
-					TRACE_CPUPROFILER_EVENT_SCOPE(FDatasmithImporter::ImportClothes:TranslatorLoadCloth);
-					if (!ImportContext.SceneTranslator->LoadCloth(ClothElement, Payload))
-					{
-						UE_LOG(LogDatasmithImport, Warning, TEXT("Cloth %s: load failed, translator issue."), ClothElement->GetName());
-						return;
-					}
+					// #ue_ds_cloth_todo fill property values, min/max,
+					PropertySetAsset->AddProperty(Property.Name, EChaosClothPresetPropertyType::String);
 				}
-				FDatasmithCloth DsCloth = MoveTemp(Payload.Cloth);
-
-				UObject* ClothAsset = nullptr;
-				TArray<UObject*> ClothPresetAssets;
+			}
+		}
 
-				// #ue_ds_cloth_arch The ImporterExtension code is expected to be temporary
-				if (IDatasmithImporterExt* Ext = IDatasmithImporterModule::Get().GetClothImporterExtension())
-				{
-					for (const FDatasmithClothPresetPropertySet& PS : DsCloth.PropertySets)
-					{
-						FString PresetAssetName = ObjectTools::SanitizeObjectName(PS.SetName);
-						if (UObject* PropertySetObject = Ext->MakeClothPropertyAsset(Outer, *PresetAssetName, ObjectFlags))
-						{
-							ClothPresetAssets.Add(PropertySetObject);
-							Ext->FillPropertySet(PropertySetObject, ClothElement, PS);
-						}
-					}
+		UObject* ExistingClothObj = nullptr;
+		if (ImportContext.SceneAsset)
+		{
+			if (TSoftObjectPtr<UObject>* ExistingAssetPtr = ImportContext.SceneAsset->Clothes.Find(ClothElement->GetName()))
+			{
+				ExistingClothObj = ExistingAssetPtr->LoadSynchronous();
+			}
+		}
+		UChaosClothAsset* ExistingCloth = Cast<UChaosClothAsset>(ExistingClothObj);
+		if (ExistingCloth)
+		{
+			if (ExistingCloth->GetOuter() != ImportPackage)
+			{
+				// Temporary flag to skip PostLoad during DuplicateObject
+				ExistingCloth->SetFlags(RF_ArchetypeObject);
+				ClothAsset = DuplicateObject<UChaosClothAsset>(ExistingCloth, ImportPackage, *ClothName);
+				ExistingCloth->ClearFlags(RF_ArchetypeObject);
+				ClothAsset->ClearFlags(RF_ArchetypeObject);
 
-					bool bShouldBuildAsset = false;
-					if (ExistingCloth)
-					{
-						if (ExistingCloth->GetOuter() != Outer)
-						{
-							// #ue_ds_todo why do we duplicate assets here ?
-							// We create a new asset from an existing instead of a new one, but we strip it from template information... Why ?
+				IDatasmithImporterModule::Get().ResetOverrides(ClothAsset);
+			}
+			else
+			{
+				ClothAsset = ExistingCloth;
+			}
 
-							// Temporary flag to skip PostLoad during DuplicateObject
-							ExistingCloth->SetFlags(RF_ArchetypeObject);
-							ClothAsset = ::DuplicateObject<UObject>(ExistingCloth, Outer, *ClothName);
-							ExistingCloth->ClearFlags(RF_ArchetypeObject);
-							ClothAsset->ClearFlags(RF_ArchetypeObject);
+			ExistingCloth->SetFlags(ObjectFlags);
+		}
+		else
+		{
+			ClothAsset = NewObject<UChaosClothAsset>(ImportPackage, *ClothName, ObjectFlags);
 
-							IDatasmithImporterModule::Get().ResetOverrides(ClothAsset);
-						}
-						else
-						{
-							ClothAsset = ExistingCloth;
-						}
+			// Setup a default skeleton
+			if (ClothAsset && !ClothAsset->GetSkeleton())
+			{
+				USkeleton* const DefaultSkeleton = LoadObject<USkeleton>(nullptr, TEXT("/Engine/EditorMeshes/SkeletalMesh/DefaultSkeletalMesh_Skeleton.DefaultSkeletalMesh_Skeleton"), nullptr, LOAD_None, nullptr);
+				check(DefaultSkeleton);
+				ClothAsset->SetSkeleton(DefaultSkeleton);
+			}
+		}
 
-						ExistingCloth->SetFlags(ObjectFlags);
-					}
-					else
-					{
-						ClothAsset = Ext->MakeClothAsset(Outer, *ClothName, ObjectFlags);
-						bShouldBuildAsset = true;
-					}
+		if (!ensure(ClothAsset))
+		{
+			continue;
+		}
 
-					Ext->FillCloth(ClothAsset, ClothElement, DsCloth); // , bShouldBuildAsset);
-				}
+		TSharedPtr<UE::Chaos::ClothAsset::FClothCollection> Collection = ClothAsset->GetClothCollection();
+		UE::Chaos::ClothAsset::FClothAdapter Adapter(Collection);
+		UE::Chaos::ClothAsset::FClothLodAdapter LODAdapter = Adapter.GetNumLods() ? Adapter.GetLod(0) : Adapter.AddGetLod();
 
-				ImportContext.ImportedClothes.Add(ClothElement, ClothAsset);
-				for (UObject* ClothPresetAsset : ClothPresetAssets)
-				{
-					ImportContext.ImportedClothPresets.Add(ClothPresetAsset);
-				}
+		for (FDatasmithClothPattern& Pattern : DsCloth.Patterns)
+		{
+			if (Pattern.IsValid())
+			{
+				UE::Chaos::ClothAsset::FClothPatternAdapter ClothPattern = LODAdapter.AddGetPattern();
+				ClothPattern.Initialize(Pattern.SimPosition, Pattern.SimRestPosition, Pattern.SimTriangleIndices);
 			}
-
-			// #ue_ds_cloth_todo: CreateStaticMeshAssetImportData
-			// #ue_ds_cloth_todo: ImportMetaDataForObject
 		}
 
-		// #ue_ds_cloth_todo: ImportContext.ImportedStaticMeshesByName.Add(MeshElement->GetName(), MeshElement); // see FDatasmithFindAssetTypeHelper
+		// Add a default material
+		ClothAsset->GetMaterials().Reset(1);
+		UMaterial* DefaultMaterial = LoadObject<UMaterial>(nullptr, TEXT("/Engine/EditorMaterials/Cloth/CameraLitDoubleSided.CameraLitDoubleSided"), nullptr, LOAD_None, nullptr);
+		const int32 MaterialId = ClothAsset->GetMaterials().Emplace(DefaultMaterial, true, false, DefaultMaterial->GetFName());
+
+		// Set a default skeleton
+		USkeleton* DefaultSkeleton = LoadObject<USkeleton>(nullptr, TEXT("/Engine/EditorMeshes/SkeletalMesh/DefaultSkeletalMesh_Skeleton.DefaultSkeletalMesh_Skeleton"), nullptr, LOAD_None, nullptr);
+		ClothAsset->SetSkeleton(DefaultSkeleton);
+
+		// Set the render mesh to duplicate the sim mesh
+		ClothAsset->CopySimMeshToRenderMesh(MaterialId);
+
+		ClothAsset->Build();
+
+		// #ue_ds_cloth_todo this step should not be necessary
+		FSkinnedAssetCompilingManager::Get().FinishCompilation({ClothAsset});
+
+		ImportContext.ImportedClothes.Add(ClothElement, ClothAsset);
+		for (UObject* ClothPresetAsset : ClothPresetAssets)
+		{
+			ImportContext.ImportedClothPresets.Add(ClothPresetAsset);
+		}
 	}
 }
 
@@ -2057,7 +2086,6 @@ void FDatasmithImporter::FinalizeImport(FDatasmithImportContext& ImportContext,
 
 	FDatasmithStaticMeshImporter::BuildStaticMeshes(StaticMeshes.Array(), ProgressFunction);
 
-	// #ue_ds_cloth_note: experimental FinalizeImport code
 	for (auto& SourcePresetObj : ImportContext.ImportedClothPresets)
 	{
 		if (!SourcePresetObj || (ValidAssets.Num() > 0 && !ValidAssets.Contains(SourcePresetObj)))
@@ -2067,7 +2095,6 @@ void FDatasmithImporter::FinalizeImport(FDatasmithImportContext& ImportContext,
 
 		FDatasmithImporterImpl::ReportProgress(Progress, 1.f, FText::FromString(FString::Printf(TEXT("Finalizing assets %d/%d (Cloth Preset %s) ..."), ++AssetIndex, NumAssetsToFinalize, *SourcePresetObj->GetName())));
 
-		// #ue_ds_cloth_existing presets
 		UObject* ExistingPreset = nullptr; // Cloth = ImportContext.SceneAsset ? ImportContext.SceneAsset->Clothes.FindOrAdd(StaticMeshId).Get() : nullptr;
 
 		FString SourcePackagePath = SourcePresetObj->GetOutermost()->GetName();
diff --git a/Engine/Plugins/Enterprise/DatasmithImporter/Source/DatasmithImporter/Private/DatasmithImporterModule.cpp b/Engine/Plugins/Enterprise/DatasmithImporter/Source/DatasmithImporter/Private/DatasmithImporterModule.cpp
--- a/Engine/Plugins/Enterprise/DatasmithImporter/Source/DatasmithImporter/Private/DatasmithImporterModule.cpp
+++ b/Engine/Plugins/Enterprise/DatasmithImporter/Source/DatasmithImporter/Private/DatasmithImporterModule.cpp
@@ -60,8 +60,6 @@
 
 #define LOCTEXT_NAMESPACE "DatasmithImporter"
 
-static class IDatasmithImporterExt* GClothImporterExtInstance = nullptr;
-
 /**
  * DatasmithImporter module implementation (private)
  */
@@ -128,18 +126,6 @@ public:
 		return GetDefault< UDatasmithImporterEditorSettings >() && GetDefault< UDatasmithImporterEditorSettings >()->bOfflineImporter;
 	}
 
-
-	virtual void SetClothImporterExtension(class IDatasmithImporterExt* InInstance) override
-	{
-		GClothImporterExtInstance = InInstance;
-	}
-
-
-	virtual class IDatasmithImporterExt* GetClothImporterExtension() override
-	{
-		return GClothImporterExtInstance;
-	}
-
 private:
 	static TSharedRef<FExtender> OnExtendContentBrowserAssetSelectionMenu(const TArray<FAssetData>& SelectedAssets);
 	static TSharedRef<FExtender> OnExtendLevelEditorActorSelectionMenu(const TSharedRef<FUICommandList> CommandList, TArray<AActor*> SelectedActors);
diff --git a/Engine/Plugins/Enterprise/DatasmithImporter/Source/DatasmithImporter/Public/DatasmithClothImporter.h b/Engine/Plugins/Enterprise/DatasmithImporter/Source/DatasmithImporter/Public/DatasmithClothImporter.h
deleted file mode 100644
diff --git a/Engine/Plugins/Enterprise/DatasmithImporter/Source/DatasmithImporter/Public/DatasmithImportContext.h b/Engine/Plugins/Enterprise/DatasmithImporter/Source/DatasmithImporter/Public/DatasmithImportContext.h
--- a/Engine/Plugins/Enterprise/DatasmithImporter/Source/DatasmithImporter/Public/DatasmithImportContext.h
+++ b/Engine/Plugins/Enterprise/DatasmithImporter/Source/DatasmithImporter/Public/DatasmithImportContext.h
@@ -257,9 +257,9 @@ struct DATASMITHIMPORTER_API FDatasmithImportContext
 	/** Map of imported mesh for each mesh element */
 	TMap< TSharedRef< IDatasmithMeshElement >, UStaticMesh* > ImportedStaticMeshes;
 
-	TMap< TSharedRef< IDatasmithClothElement >, UObject* > ImportedClothes; // #ue_ds_cloth_arch UChaosClothAsset
+	TMap< TSharedRef< IDatasmithClothElement >, UObject* > ImportedClothes;
 
-	TArray< UObject* > ImportedClothPresets; // #ue_ds_cloth_arch: UChaosClothPreset // #ue_ds_cloth_todo: map with a dedicated element so that clothes can share presets // #ue_ds_cloth_existing presets
+	TArray< UObject* > ImportedClothPresets;
 
 	/** Register IDatasmithMeshElement by their name so they can be searched faster */
 	TMap< FString, TSharedRef < IDatasmithMeshElement > > ImportedStaticMeshesByName;
diff --git a/Engine/Plugins/Enterprise/DatasmithImporter/Source/DatasmithImporter/Public/DatasmithImporterModule.h b/Engine/Plugins/Enterprise/DatasmithImporter/Source/DatasmithImporter/Public/DatasmithImporterModule.h
--- a/Engine/Plugins/Enterprise/DatasmithImporter/Source/DatasmithImporter/Public/DatasmithImporterModule.h
+++ b/Engine/Plugins/Enterprise/DatasmithImporter/Source/DatasmithImporter/Public/DatasmithImporterModule.h
@@ -45,9 +45,5 @@ public:
 
 	DECLARE_MULTICAST_DELEGATE_OneParam(FOnGenerateDatasmithImportMenu, struct FToolMenuSection&);
 	virtual FOnGenerateDatasmithImportMenu& OnGenerateDatasmithImportMenu() = 0;
-
-	// #ue_ds_cloth_arch: temp API: Do not use directly.
-	virtual void SetClothImporterExtension(class IDatasmithImporterExt*) = 0;
-	virtual class IDatasmithImporterExt* GetClothImporterExtension() = 0;
 };
 
diff --git a/Engine/Source/Runtime/Datasmith/DatasmithCore/Private/DatasmithSceneXmlReader.cpp b/Engine/Source/Runtime/Datasmith/DatasmithCore/Private/DatasmithSceneXmlReader.cpp
--- a/Engine/Source/Runtime/Datasmith/DatasmithCore/Private/DatasmithSceneXmlReader.cpp
+++ b/Engine/Source/Runtime/Datasmith/DatasmithCore/Private/DatasmithSceneXmlReader.cpp
@@ -439,7 +439,6 @@ void FDatasmithSceneXmlReader::ParseMesh(FXmlNode* InNode, TSharedPtr<IDatasmith
 
 void FDatasmithSceneXmlReader::ParseCloth(FXmlNode* InNode, TSharedPtr<IDatasmithClothElement>& OutElement) const
 {
-	// #ue_ds_cloth_note: xml read
 	ParseElement( InNode, OutElement.ToSharedRef() );
 
 	for (const FXmlNode* Node : InNode->GetChildrenNodes())
diff --git a/Engine/Source/Runtime/Datasmith/DatasmithCore/Private/DatasmithSceneXmlWriter.cpp b/Engine/Source/Runtime/Datasmith/DatasmithCore/Private/DatasmithSceneXmlWriter.cpp
--- a/Engine/Source/Runtime/Datasmith/DatasmithCore/Private/DatasmithSceneXmlWriter.cpp
+++ b/Engine/Source/Runtime/Datasmith/DatasmithCore/Private/DatasmithSceneXmlWriter.cpp
@@ -362,7 +362,6 @@ void FDatasmithSceneXmlWriterImpl::WriteMeshElement(const TSharedPtr< IDatasmith
 
 void FDatasmithSceneXmlWriterImpl::WriteClothElement(const TSharedPtr<IDatasmithClothElement>& Element, FArchive& Archive, int32 Indent)
 {
-	// #ue_ds_cloth_note: FDatasmithSceneXmlWriterImpl::WriteClothElement
 	WriteIndent(Archive, Indent);
 	FString XmlString = TEXT("<") DATASMITH_CLOTH;
 	XmlString += TEXT(" name=\"") + SanitizeXMLText( Element->GetName() ) + TEXT("\"");
diff --git a/Engine/Source/Runtime/Datasmith/DatasmithCore/Public/IDatasmithSceneElements.h b/Engine/Source/Runtime/Datasmith/DatasmithCore/Public/IDatasmithSceneElements.h
--- a/Engine/Source/Runtime/Datasmith/DatasmithCore/Public/IDatasmithSceneElements.h
+++ b/Engine/Source/Runtime/Datasmith/DatasmithCore/Public/IDatasmithSceneElements.h
@@ -238,7 +238,7 @@ protected:
 };
 
 /**
-* #ue_ds_cloth_doc IDatasmithClothElement class: experimental class that describes a cloth asset
+* IDatasmithClothElement class: experimental class that describes a cloth asset
 */
 class DATASMITHCORE_API IDatasmithClothElement : public IDatasmithElement
 {
@@ -1610,7 +1610,7 @@ public:
 	*/
 	virtual void EmptyMeshes() = 0;
 
-	// #ue_ds_cloth_doc IDatasmithScene API
+	// #ue_ds_todo cloth api doc
 	virtual void AddCloth(const TSharedPtr< IDatasmithClothElement >& InElement) = 0;
 	virtual int32 GetClothesCount() const = 0;
 	virtual TSharedPtr< IDatasmithClothElement > GetCloth(int32 InIndex) = 0;
