Datasmith: When reimporting reuse components of exactly same UClass

- e.g. using HISM component for StaticMesh doesn't work(although HISM subclasses StaticMesh)

diff --git a/Engine/Plugins/Enterprise/DatasmithImporter/Source/DatasmithImporter/Private/Utility/DatasmithImporterImpl.cpp b/Engine/Plugins/Enterprise/DatasmithImporter/Source/DatasmithImporter/Private/Utility/DatasmithImporterImpl.cpp
--- a/Engine/Plugins/Enterprise/DatasmithImporter/Source/DatasmithImporter/Private/Utility/DatasmithImporterImpl.cpp
+++ b/Engine/Plugins/Enterprise/DatasmithImporter/Source/DatasmithImporter/Private/Utility/DatasmithImporterImpl.cpp
@@ -568,7 +568,9 @@ UActorComponent* FDatasmithImporterImpl::PublicizeComponent(UActorComponent& Sou
 
 USceneComponent* FDatasmithImporterImpl::FinalizeSceneComponent(FDatasmithImportContext& ImportContext, USceneComponent& SourceComponent, AActor& DestinationActor, USceneComponent* DestinationParent, TMap<UObject *, UObject *>& ReferencesToRemap, TArray<uint8>& ReusableBuffer, TArray<TPair<USceneComponent&, TArray<FMigratedTemplatePairType>>>& ComponentsToApplyMigratedTemplate)
 {
-	USceneComponent* DestinationComponent = static_cast<USceneComponent*>( FindObjectWithOuter( &DestinationActor, SourceComponent.GetClass(), SourceComponent.GetFName() ) );
+	bool bExactClass = true; // Need exact class (UHierarchicalInstancedStaticMeshComponent should not be reused for UStaticMeshComponent)
+	USceneComponent* DestinationComponent = static_cast<USceneComponent*>( StaticFindObjectFastSafe(SourceComponent.GetClass(), &DestinationActor, SourceComponent.GetFName(), bExactClass) );
+	
 	FName SourceComponentDatasmithId = FDatasmithImporterUtils::GetDatasmithElementId(&SourceComponent);
 
 	if ( SourceComponentDatasmithId.IsNone() )
@@ -596,7 +598,7 @@ USceneComponent* FDatasmithImporterImpl::FinalizeSceneComponent(FDatasmithImport
 			// Look at components of the actor, we might find the scene component we are looking for.
 			for ( UActorComponent* Component : DestinationActor.GetInstanceComponents() )
 			{
-				if ( Component && Component->IsA( SourceComponent.GetClass() ) )
+				if ( Component && ( Component->GetClass() == SourceComponent.GetClass() ) )
 				{
 					FName ComponentDatasmithId = FDatasmithImporterUtils::GetDatasmithElementId( Component );
 					if ( ComponentDatasmithId.IsEqual( SourceComponentDatasmithId ) )
