[Datasmith] Direct Link -  reconnect to source for previously imported asset

- this was a regression - worked with UE 5.0

diff --git a/Engine/Plugins/Enterprise/DatasmithImporter/Source/DirectLinkExtension/Private/AutoReimportManager.h b/Engine/Plugins/Enterprise/DatasmithImporter/Source/DirectLinkExtension/Private/AutoReimportManager.h
--- a/Engine/Plugins/Enterprise/DatasmithImporter/Source/DirectLinkExtension/Private/AutoReimportManager.h
+++ b/Engine/Plugins/Enterprise/DatasmithImporter/Source/DirectLinkExtension/Private/AutoReimportManager.h
@@ -23,7 +23,7 @@ namespace UE::DatasmithImporter
 	public:
 		/* Begin FTickableGameObject Interface */
 		virtual void Tick(float DeltaTime) override;
-		virtual bool IsAllowedToTick() const override { return PendingAutoReimportObjects.Num() > 0 || !PendingReimportQueue.IsEmpty(); }
+		virtual bool IsAllowedToTick() const override { return PendingAutoReimportObjects.Num() > 0 || !PendingReimportQueue.IsEmpty() || !PendingInvalidateQueue.IsEmpty(); }
 		virtual TStatId GetStatId() const override;
 		virtual bool IsTickableWhenPaused() const override { return true; }
 		virtual bool IsTickableInEditor() const override { return true; }
