DatasmithDirectLink - Fix crash on shutdown that could happen when receiving a DirectLink message while the engine is shutting down.

Under the wrong circumstances the engine could crash on shutdown due to a race condition in the DirectLink async thread where a static <UdpMessageType>::StaticStruct() could be freed while using it in the CommonMessaging module.
This could happen because no explicit dependency was declared between the UdpMessaging module and the DirectLinkExtension module.
This fix adds the missing dependency and so UdpMessaging will now always be unloaded after DirectLinkExtension, fixing the use-after free crash.

robomerge[FN_EngineMerge] -UE5-RES

diff --git a/Engine/Plugins/Enterprise/DatasmithImporter/DatasmithImporter.uplugin b/Engine/Plugins/Enterprise/DatasmithImporter/DatasmithImporter.uplugin
--- a/Engine/Plugins/Enterprise/DatasmithImporter/DatasmithImporter.uplugin
+++ b/Engine/Plugins/Enterprise/DatasmithImporter/DatasmithImporter.uplugin
@@ -88,7 +88,11 @@
 			"Enabled": true
 		},
 		{
-			"Name": "ChaosClothAsset",
+			"Name": "UdpMessaging",
+			"Enabled": true
+		},
+		{
+			"Name": "UdpMessaging",
 			"Enabled": true
 		}
 	]
diff --git a/Engine/Plugins/Enterprise/DatasmithImporter/Source/DirectLinkExtension/DirectLinkExtension.Build.cs b/Engine/Plugins/Enterprise/DatasmithImporter/Source/DirectLinkExtension/DirectLinkExtension.Build.cs
--- a/Engine/Plugins/Enterprise/DatasmithImporter/Source/DirectLinkExtension/DirectLinkExtension.Build.cs
+++ b/Engine/Plugins/Enterprise/DatasmithImporter/Source/DirectLinkExtension/DirectLinkExtension.Build.cs
@@ -14,6 +14,7 @@ namespace UnrealBuildTool.Rules
 				{
 					"Engine",
 					"InputCore",
+					"UdpMessaging", // required for DirectLink networking
 				}
 			);
 
