Fix merge issue from CL 22891140 causing Plugin warning: ['DatasmithImporter' does not list plugin 'ChaosClothAsset' as a dependency, but module 'DatasmithImporter' depends on 'ChaosClothAsset'.]

diff --git a/Engine/Plugins/Enterprise/DatasmithImporter/DatasmithImporter.uplugin b/Engine/Plugins/Enterprise/DatasmithImporter/DatasmithImporter.uplugin
--- a/Engine/Plugins/Enterprise/DatasmithImporter/DatasmithImporter.uplugin
+++ b/Engine/Plugins/Enterprise/DatasmithImporter/DatasmithImporter.uplugin
@@ -92,7 +92,7 @@
 			"Enabled": true
 		},
 		{
-			"Name": "UdpMessaging",
+			"Name": "ChaosClothAsset",
 			"Enabled": true
 		}
 	]
