DatasmithCloth: import Sewing info
- fix reimport data keeping old patterns
- fix usage of reference skeleton API
- skip import of empty clothes (with no patterns) as it causes crashes on asset build

diff --git a/Engine/Plugins/Enterprise/DatasmithImporter/Source/DatasmithImporter/Private/DatasmithImporter.cpp b/Engine/Plugins/Enterprise/DatasmithImporter/Source/DatasmithImporter/Private/DatasmithImporter.cpp
--- a/Engine/Plugins/Enterprise/DatasmithImporter/Source/DatasmithImporter/Private/DatasmithImporter.cpp
+++ b/Engine/Plugins/Enterprise/DatasmithImporter/Source/DatasmithImporter/Private/DatasmithImporter.cpp
@@ -427,6 +427,12 @@ void FDatasmithImporter::ImportClothes(FDatasmithImportContext& ImportContext)
 				UE_LOG(LogDatasmithImport, Warning, TEXT("Cloth %s: load failed, translator issue."), ClothElement->GetName());
 				return;
 			}
+
+			if (Payload.Cloth.Patterns.IsEmpty())
+			{
+				UE_LOG(LogDatasmithImport, Warning, TEXT("Cloth %s: no patterns, import skipped."), ClothElement->GetName());
+				return;
+			}
 		}
 		FDatasmithCloth DsCloth = MoveTemp(Payload.Cloth);
 
@@ -498,6 +504,7 @@ void FDatasmithImporter::ImportClothes(FDatasmithImportContext& ImportContext)
 
 		TSharedPtr<UE::Chaos::ClothAsset::FClothCollection> Collection = ClothAsset->GetClothCollection();
 		UE::Chaos::ClothAsset::FClothAdapter Adapter(Collection);
+		Adapter.Reset(); // on reimport, asset is non-empty
 		UE::Chaos::ClothAsset::FClothLodAdapter LODAdapter = Adapter.GetNumLods() ? Adapter.GetLod(0) : Adapter.AddGetLod();
 
 		for (FDatasmithClothPattern& Pattern : DsCloth.Patterns)
@@ -509,23 +516,41 @@ void FDatasmithImporter::ImportClothes(FDatasmithImportContext& ImportContext)
 			}
 		}
 
+		LODAdapter.SetNumSeams(DsCloth.Sewing.Num());
+		TArrayView<FIntVector2> SeamPatterns = LODAdapter.GetSeamPatterns();
+		TArrayView<TArray<FIntVector2>> SeamStitches = LODAdapter.GetSeamStitches();
+
+		uint32 SeamIndex = 0;
+		for (const FDatasmithClothSewingInfo& SeamInfo : DsCloth.Sewing)
+		{
+			SeamPatterns[SeamIndex] = FIntVector2(SeamInfo.Seam0PanelIndex, SeamInfo.Seam1PanelIndex);
+			TArray<FIntVector2>& Stitches = SeamStitches[SeamIndex];
+			uint32 StitchesCount = std::min(SeamInfo.Seam0MeshIndices.Num(), SeamInfo.Seam1MeshIndices.Num());
+			Stitches.Reserve(StitchesCount);
+			for (uint32 StitchIndex = 0; StitchIndex < StitchesCount; ++StitchIndex)
+			{
+				Stitches.Add(FIntVector2(SeamInfo.Seam0MeshIndices[StitchIndex], SeamInfo.Seam1MeshIndices[StitchIndex]));
+			}
+
+			++SeamIndex;
+		}
+
 		// Add a default material
 		ClothAsset->GetMaterials().Reset(1);
 		UMaterial* DefaultMaterial = LoadObject<UMaterial>(nullptr, TEXT("/Engine/EditorMaterials/Cloth/CameraLitDoubleSided.CameraLitDoubleSided"), nullptr, LOAD_None, nullptr);
 		const int32 MaterialId = ClothAsset->GetMaterials().Emplace(DefaultMaterial, true, false, DefaultMaterial->GetFName());
 
 		// Set a default skeleton
-		USkeleton* DefaultSkeleton = LoadObject<USkeleton>(nullptr, TEXT("/Engine/EditorMeshes/SkeletalMesh/DefaultSkeletalMesh_Skeleton.DefaultSkeletalMesh_Skeleton"), nullptr, LOAD_None, nullptr);
-		ClothAsset->SetSkeleton(DefaultSkeleton);
+		if (USkeleton* DefaultSkeleton = LoadObject<USkeleton>(nullptr, TEXT("/Engine/EditorMeshes/SkeletalMesh/DefaultSkeletalMesh_Skeleton.DefaultSkeletalMesh_Skeleton"), nullptr, LOAD_None, nullptr))
+		{
+			ClothAsset->SetReferenceSkeleton(DefaultSkeleton->GetReferenceSkeleton(), false);
+		}
 
 		// Set the render mesh to duplicate the sim mesh
 		ClothAsset->CopySimMeshToRenderMesh(MaterialId);
 
 		ClothAsset->Build();
 
-		// #ue_ds_cloth_todo this step should not be necessary
-		FSkinnedAssetCompilingManager::Get().FinishCompilation({ClothAsset});
-
 		ImportContext.ImportedClothes.Add(ClothElement, ClothAsset);
 		for (UObject* ClothPresetAsset : ClothPresetAssets)
 		{
diff --git a/Engine/Source/Runtime/Datasmith/DatasmithCore/Private/DatasmithCloth.cpp b/Engine/Source/Runtime/Datasmith/DatasmithCore/Private/DatasmithCloth.cpp
--- a/Engine/Source/Runtime/Datasmith/DatasmithCore/Private/DatasmithCloth.cpp
+++ b/Engine/Source/Runtime/Datasmith/DatasmithCore/Private/DatasmithCloth.cpp
@@ -10,10 +10,11 @@ enum EDatasmithClothSerializationVersion
 {
     EDCSV_Base = 0,
     EDCSV_WithPatternParameters = 1,
+    EDCSV_WithSewingInfo = 2,
 
 	// -----<new versions can be added before this line>-------------------------------------------------
-    _EDCSV_LastPlusOne,
-    _EDCSV_Last = _EDCSV_LastPlusOne -1
+    _EDCSV_Count,
+    _EDCSV_Last = _EDCSV_Count - 1
 };
 
 struct FDatasmithClothSerializationVersion
@@ -24,7 +25,6 @@ private:
     FDatasmithClothSerializationVersion() = default;
 };
 
-
 const FGuid FDatasmithClothSerializationVersion::GUID(0x28B01036, 0x66B4498F, 0x99425ACA, 0xDB78A9B5);
 
 // Register the custom version with core
@@ -45,10 +45,11 @@ FArchive& operator<<(FArchive& Ar, FParameterData& ParameterData)
 	return Ar;
 }
 
+
 FArchive& operator<<(FArchive& Ar, FDatasmithClothPattern& Pattern)
 {
 	Ar.UsingCustomVersion(FDatasmithClothSerializationVersion::GUID);
-	int32 ClothSerialVersion = Ar.CustomVer(FDatasmithClothSerializationVersion::GUID);
+	const int32 ClothSerialVersion = Ar.CustomVer(FDatasmithClothSerializationVersion::GUID);
 
 	Ar << Pattern.SimPosition;
 	Ar << Pattern.SimRestPosition;
@@ -63,6 +64,19 @@ FArchive& operator<<(FArchive& Ar, FDatasmithClothPattern& Pattern)
 }
 
 
+FArchive& operator<<(FArchive& Ar, FDatasmithClothSewingInfo& Sewing)
+{
+	Ar.UsingCustomVersion(FDatasmithClothSerializationVersion::GUID);
+
+	Ar << Sewing.Seam0MeshIndices;
+	Ar << Sewing.Seam1MeshIndices;
+	Ar << Sewing.Seam0PanelIndex;
+	Ar << Sewing.Seam1PanelIndex;
+
+	return Ar;
+}
+
+
 FArchive& operator<<(FArchive& Ar, FDatasmithClothPresetProperty& Property)
 {
 	Ar.UsingCustomVersion(FDatasmithClothSerializationVersion::GUID);
@@ -73,6 +87,7 @@ FArchive& operator<<(FArchive& Ar, FDatasmithClothPresetProperty& Property)
 	return Ar;
 }
 
+
 FArchive& operator<<(FArchive& Ar, FDatasmithClothPresetPropertySet& PropertySet)
 {
 	Ar.UsingCustomVersion(FDatasmithClothSerializationVersion::GUID);
@@ -83,12 +98,19 @@ FArchive& operator<<(FArchive& Ar, FDatasmithClothPresetPropertySet& PropertySet
 	return Ar;
 }
 
+
 FArchive& operator<<(FArchive& Ar, FDatasmithCloth& Cloth)
 {
 	Ar.UsingCustomVersion(FDatasmithClothSerializationVersion::GUID);
+	const int32 ClothSerialVersion = Ar.CustomVer(FDatasmithClothSerializationVersion::GUID);
 
 	Ar << Cloth.Patterns;
 	Ar << Cloth.PropertySets;
 
+	if (ClothSerialVersion >= EDCSV_WithSewingInfo)
+	{
+		Ar << Cloth.Sewing;
+	}
+
 	return Ar;
 }
diff --git a/Engine/Source/Runtime/Datasmith/DatasmithCore/Public/DatasmithCloth.h b/Engine/Source/Runtime/Datasmith/DatasmithCore/Public/DatasmithCloth.h
--- a/Engine/Source/Runtime/Datasmith/DatasmithCore/Public/DatasmithCloth.h
+++ b/Engine/Source/Runtime/Datasmith/DatasmithCore/Public/DatasmithCloth.h
@@ -4,7 +4,6 @@
 #include "Containers/Array.h"
 #include "Containers/UnrealString.h"
 #include "CoreMinimal.h"
-#include "Math/UnrealMathSSE.h"
 #include "Math/Vector.h"
 #include "Math/Vector2D.h"
 #include "Misc/TVariant.h"
@@ -23,7 +22,7 @@ struct FParameterData
 	TVariant<TArray<float>, TArray<double>> Data;
 
 public:
-	friend FArchive& operator<<(FArchive& Ar, FParameterData& Pattern);
+	friend FArchive& operator<<(FArchive& Ar, FParameterData& ParameterData);
 };
 
 
@@ -37,11 +36,22 @@ public:
 	TArray<FParameterData> Parameters;
 
 public:
-	bool IsValid() { return SimRestPosition.Num() == SimPosition.Num() && SimTriangleIndices.Num() % 3 == 0 && SimTriangleIndices.Num(); }
+	bool IsValid() const { return SimRestPosition.Num() == SimPosition.Num() && SimTriangleIndices.Num() % 3 == 0 && SimTriangleIndices.Num(); }
 	friend FArchive& operator<<(FArchive& Ar, FDatasmithClothPattern& Pattern);
 };
 
 
+struct FDatasmithClothSewingInfo
+{
+	uint32 Seam0PanelIndex = 0;
+	uint32 Seam1PanelIndex = 0;
+	TArray<uint32> Seam0MeshIndices;
+	TArray<uint32> Seam1MeshIndices;
+
+	friend FArchive& operator<<(FArchive& Ar, FDatasmithClothSewingInfo& Sewing);
+};
+
+
 class FDatasmithClothPresetProperty
 {
 public:
@@ -68,6 +78,7 @@ class DATASMITHCORE_API FDatasmithCloth
 {
 public:
 	TArray<FDatasmithClothPattern> Patterns;
+	TArray<FDatasmithClothSewingInfo> Sewing;
 	TArray<FDatasmithClothPresetPropertySet> PropertySets;
 
 public:
