Fixed crash when commiting a HSIM with null material overrides

diff --git a/Engine/Plugins/Enterprise/DatasmithImporter/Source/DatasmithImporter/Private/Utility/DatasmithImporterUtils.cpp b/Engine/Plugins/Enterprise/DatasmithImporter/Source/DatasmithImporter/Private/Utility/DatasmithImporterUtils.cpp
--- a/Engine/Plugins/Enterprise/DatasmithImporter/Source/DatasmithImporter/Private/Utility/DatasmithImporterUtils.cpp
+++ b/Engine/Plugins/Enterprise/DatasmithImporter/Source/DatasmithImporter/Private/Utility/DatasmithImporterUtils.cpp
@@ -1016,7 +1016,11 @@ namespace FDatasmithImporterUtilsHelper
 
 					for(UMaterialInterface* MaterialInterface : HISMComponent->OverrideMaterials)
 					{
-						HISMActorElement->AddMaterialOverride( *MaterialInterface->GetName(), 0 );
+						// Quick fix to avoid a crash. The logic for material override seems bogus for HISMComponent and StaticMeshComponent
+						if (MaterialInterface)
+						{
+							HISMActorElement->AddMaterialOverride(*MaterialInterface->GetName(), 0);
+						}
 					}
 
 					FString StaticMeshTag = FDatasmithImporterUtils::GetDatasmithElementIdString( HISMComponent->GetStaticMesh() );
