Chaos Cloth - Remove Cloth Presets now superseded by the Cloth Collection.

#rnx

diff --git a/Engine/Plugins/Enterprise/DatasmithImporter/Source/DatasmithImporter/Private/DatasmithImporter.cpp b/Engine/Plugins/Enterprise/DatasmithImporter/Source/DatasmithImporter/Private/DatasmithImporter.cpp
--- a/Engine/Plugins/Enterprise/DatasmithImporter/Source/DatasmithImporter/Private/DatasmithImporter.cpp
+++ b/Engine/Plugins/Enterprise/DatasmithImporter/Source/DatasmithImporter/Private/DatasmithImporter.cpp
@@ -82,7 +82,6 @@
 #include "UnrealEdGlobals.h"
 
 #include "Animation/Skeleton.h"
-#include "ChaosClothAsset/ClothPreset.h"
 #include "ChaosClothAsset/ClothAsset.h"
 #include "ChaosClothAsset/ClothAdapter.h"
 #include "ChaosClothAsset/ClothLodAdapter.h"
@@ -440,21 +439,23 @@ void FDatasmithImporter::ImportClothes(FDatasmithImportContext& ImportContext)
 		EObjectFlags ObjectFlags = ImportContext.ObjectFlags & ~RF_Public; // not RF_Public yet, the publicized asset will be.
 
 		UChaosClothAsset* ClothAsset = nullptr;
-		TArray<UObject*> ClothPresetAssets;
 
-		for (const FDatasmithClothPresetPropertySet& PropertySet : DsCloth.PropertySets)
-		{
-			FString PresetAssetName = ObjectTools::SanitizeObjectName(PropertySet.SetName);
-			if (UChaosClothPreset* PropertySetAsset = NewObject<UChaosClothPreset>(ImportPackage, *PresetAssetName, ObjectFlags))
-			{
-				ClothPresetAssets.Add(PropertySetAsset);
-				for (const FDatasmithClothPresetProperty& Property : PropertySet.Properties)
-				{
-					// #ue_ds_cloth_todo fill property values, min/max,
-					PropertySetAsset->AddProperty(Property.Name, EChaosClothPresetPropertyType::String);
-				}
-			}
-		}
+		// TODO: Use PropertyCollectionAdapter instead, see \Engine\Source\Runtime\Experimental\Chaos\Public\Chaos\PropertyCollectionAdapter.h
+		//TArray<UObject*> ClothPresetAssets;
+
+		//for (const FDatasmithClothPresetPropertySet& PropertySet : DsCloth.PropertySets)
+		//{
+		//	FString PresetAssetName = ObjectTools::SanitizeObjectName(PropertySet.SetName);
+		//	if (UChaosClothPreset* PropertySetAsset = NewObject<UChaosClothPreset>(ImportPackage, *PresetAssetName, ObjectFlags))
+		//	{
+		//		ClothPresetAssets.Add(PropertySetAsset);
+		//		for (const FDatasmithClothPresetProperty& Property : PropertySet.Properties)
+		//		{
+		//			// #ue_ds_cloth_todo fill property values, min/max,
+		//			PropertySetAsset->AddProperty(Property.Name, EChaosClothPresetPropertyType::String);
+		//		}
+		//	}
+		//}
 
 		UObject* ExistingClothObj = nullptr;
 		if (ImportContext.SceneAsset)
@@ -552,10 +553,12 @@ void FDatasmithImporter::ImportClothes(FDatasmithImportContext& ImportContext)
 		ClothAsset->Build();
 
 		ImportContext.ImportedClothes.Add(ClothElement, ClothAsset);
-		for (UObject* ClothPresetAsset : ClothPresetAssets)
-		{
-			ImportContext.ImportedClothPresets.Add(ClothPresetAsset);
-		}
+
+		// TODO: Use PropertyCollectionAdapter instead, see \Engine\Source\Runtime\Experimental\Chaos\Public\Chaos\PropertyCollectionAdapter.h
+		//for (UObject* ClothPresetAsset : ClothPresetAssets)
+		//{
+		//	ImportContext.ImportedClothPresets.Add(ClothPresetAsset);
+		//}
 	}
 }
 
diff --git a/Engine/Plugins/Experimental/ChaosClothAsset/Source/ChaosClothAsset/Private/ChaosClothAsset/ClothPresetCollection.cpp b/Engine/Plugins/Experimental/ChaosClothAsset/Source/ChaosClothAsset/Private/ChaosClothAsset/ClothPresetCollection.cpp
deleted file mode 100644
diff --git a/Engine/Plugins/Experimental/ChaosClothAsset/Source/ChaosClothAsset/Public/ChaosClothAsset/ClothPresetCollection.h b/Engine/Plugins/Experimental/ChaosClothAsset/Source/ChaosClothAsset/Public/ChaosClothAsset/ClothPresetCollection.h
deleted file mode 100644
diff --git a/Engine/Plugins/Experimental/ChaosClothAsset/Source/ChaosClothAssetEngine/Private/ChaosClothAsset/ClothPreset.cpp b/Engine/Plugins/Experimental/ChaosClothAsset/Source/ChaosClothAssetEngine/Private/ChaosClothAsset/ClothPreset.cpp
deleted file mode 100644
diff --git a/Engine/Plugins/Experimental/ChaosClothAsset/Source/ChaosClothAssetEngine/Public/ChaosClothAsset/ClothAsset.h b/Engine/Plugins/Experimental/ChaosClothAsset/Source/ChaosClothAssetEngine/Public/ChaosClothAsset/ClothAsset.h
--- a/Engine/Plugins/Experimental/ChaosClothAsset/Source/ChaosClothAssetEngine/Public/ChaosClothAsset/ClothAsset.h
+++ b/Engine/Plugins/Experimental/ChaosClothAsset/Source/ChaosClothAssetEngine/Public/ChaosClothAsset/ClothAsset.h
@@ -2,7 +2,6 @@
 
 #pragma once
 
-#include "ChaosClothAsset/ClothPreset.h"
 #include "Engine/SkinnedAsset.h"
 #include "ReferenceSkeleton.h"
 #include "RenderCommandFence.h"
@@ -186,10 +185,6 @@ private:
 	/** Reregister all components using this asset to reset the simulation in case anything has changed. */
 	void ReregisterComponents();
 
-	/** List of cloth presets for this cloth asset. */
-	UPROPERTY(EditAnywhere, AssetRegistrySearchable, Category = ClothPresets)
-	TArray<TObjectPtr<UChaosClothPreset>> ClothPresets;
-
 	// TODO: determine if it should be serialized or transient
 	/** List of materials for this cloth asset. */
 	UPROPERTY(EditAnywhere, Category = Materials)
diff --git a/Engine/Plugins/Experimental/ChaosClothAsset/Source/ChaosClothAssetEngine/Public/ChaosClothAsset/ClothPreset.h b/Engine/Plugins/Experimental/ChaosClothAsset/Source/ChaosClothAssetEngine/Public/ChaosClothAsset/ClothPreset.h
deleted file mode 100644
diff --git a/Engine/Plugins/Experimental/ChaosClothAssetEditor/Source/ChaosClothAssetEditor/Private/ChaosClothAsset/AssetTypeActions_ClothPreset.cpp b/Engine/Plugins/Experimental/ChaosClothAssetEditor/Source/ChaosClothAssetEditor/Private/ChaosClothAsset/AssetTypeActions_ClothPreset.cpp
deleted file mode 100644
diff --git a/Engine/Plugins/Experimental/ChaosClothAssetEditor/Source/ChaosClothAssetEditor/Private/ChaosClothAsset/AssetTypeActions_ClothPreset.h b/Engine/Plugins/Experimental/ChaosClothAssetEditor/Source/ChaosClothAssetEditor/Private/ChaosClothAsset/AssetTypeActions_ClothPreset.h
deleted file mode 100644
diff --git a/Engine/Plugins/Experimental/ChaosClothAssetEditor/Source/ChaosClothAssetEditor/Private/ChaosClothAsset/ChaosClothAssetEditorModule.cpp b/Engine/Plugins/Experimental/ChaosClothAssetEditor/Source/ChaosClothAssetEditor/Private/ChaosClothAsset/ChaosClothAssetEditorModule.cpp
--- a/Engine/Plugins/Experimental/ChaosClothAssetEditor/Source/ChaosClothAssetEditor/Private/ChaosClothAsset/ChaosClothAssetEditorModule.cpp
+++ b/Engine/Plugins/Experimental/ChaosClothAssetEditor/Source/ChaosClothAssetEditor/Private/ChaosClothAsset/ChaosClothAssetEditorModule.cpp
@@ -13,7 +13,6 @@
 #include "Styling/AppStyle.h"
 #include "AssetToolsModule.h"
 #include "AssetTypeActions_ClothAsset.h"
-#include "AssetTypeActions_ClothPreset.h"
 
 #define LOCTEXT_NAMESPACE "FChaosClothAssetEditorModule"
 
@@ -32,9 +31,6 @@ void FChaosClothAssetEditorModule::StartupModule()
 	AssetTypeActions_ClothAsset = new FAssetTypeActions_ClothAsset();
 	AssetTools.RegisterAssetTypeActions(MakeShareable(AssetTypeActions_ClothAsset));
 
-	AssetTypeActions_ClothPreset = new FAssetTypeActions_ClothPreset();
-	AssetTools.RegisterAssetTypeActions(MakeShareable(AssetTypeActions_ClothPreset));
-
 	Dataflow::RegisterClothDataflowNodes();
 
 	// TODO: Register details view customizations
@@ -52,7 +48,6 @@ void FChaosClothAssetEditorModule::ShutdownModule()
 		FAssetToolsModule& AssetToolsModule = FAssetToolsModule::GetModule();
 		IAssetTools& AssetTools = AssetToolsModule.Get();
 		AssetTools.UnregisterAssetTypeActions(AssetTypeActions_ClothAsset->AsShared());
-		AssetTools.UnregisterAssetTypeActions(AssetTypeActions_ClothPreset->AsShared());
 	}
 
 	// TODO: Unregister details view customizations
diff --git a/Engine/Plugins/Experimental/ChaosClothAssetEditor/Source/ChaosClothAssetTools/Private/ChaosClothAsset/ChaosClothAssetToolsModule.cpp b/Engine/Plugins/Experimental/ChaosClothAssetEditor/Source/ChaosClothAssetTools/Private/ChaosClothAsset/ChaosClothAssetToolsModule.cpp
--- a/Engine/Plugins/Experimental/ChaosClothAssetEditor/Source/ChaosClothAssetTools/Private/ChaosClothAsset/ChaosClothAssetToolsModule.cpp
+++ b/Engine/Plugins/Experimental/ChaosClothAssetEditor/Source/ChaosClothAssetTools/Private/ChaosClothAsset/ChaosClothAssetToolsModule.cpp
@@ -7,7 +7,6 @@
 #include "AssetToolsModule.h"
 #include "ClothAssetEditorStyle.h"
 #include "ClothComponentEditorStyle.h"
-#include "ClothPresetEditorStyle.h"
 #include "ClothingAssetToClothAssetExporter.h"
 #include "ChaosClothAsset/ClothAssetBuilderEditor.h"
 
@@ -22,7 +21,6 @@ namespace UE::Chaos::ClothAsset
 			// Register asset icons
 			FClothAssetEditorStyle::Get();
 			FClothComponentEditorStyle::Get();
-			FClothPresetEditorStyle::Get();
 
 			// Register modular features
 			IModularFeatures::Get().RegisterModularFeature(IClothingAssetExporterClassProvider::FeatureName, static_cast<IClothingAssetExporterClassProvider*>(this));
diff --git a/Engine/Plugins/Experimental/ChaosClothAssetEditor/Source/ChaosClothAssetTools/Private/ChaosClothAsset/ClothPresetEditorStyle.cpp b/Engine/Plugins/Experimental/ChaosClothAssetEditor/Source/ChaosClothAssetTools/Private/ChaosClothAsset/ClothPresetEditorStyle.cpp
deleted file mode 100644
diff --git a/Engine/Plugins/Experimental/ChaosClothAssetEditor/Source/ChaosClothAssetTools/Private/ChaosClothAsset/ClothPresetEditorStyle.h b/Engine/Plugins/Experimental/ChaosClothAssetEditor/Source/ChaosClothAssetTools/Private/ChaosClothAsset/ClothPresetEditorStyle.h
deleted file mode 100644
diff --git a/Engine/Plugins/Experimental/ChaosClothAssetEditor/Source/ChaosClothAssetTools/Private/ChaosClothAsset/ClothPresetFactory.cpp b/Engine/Plugins/Experimental/ChaosClothAssetEditor/Source/ChaosClothAssetTools/Private/ChaosClothAsset/ClothPresetFactory.cpp
deleted file mode 100644
diff --git a/Engine/Plugins/Experimental/ChaosClothAssetEditor/Source/ChaosClothAssetTools/Private/ChaosClothAsset/ClothPresetFactory.h b/Engine/Plugins/Experimental/ChaosClothAssetEditor/Source/ChaosClothAssetTools/Private/ChaosClothAsset/ClothPresetFactory.h
deleted file mode 100644
