Datasmith:
- Added option for Datasmith Actor Factory to use the passed InLevel parameter from the UFactory interface, instead of leaving ImportWorld nullptr and just using GWorld

diff --git a/Engine/Plugins/Enterprise/DatasmithImporter/Source/DatasmithImporter/Private/ActorFactoryDatasmithScene.cpp b/Engine/Plugins/Enterprise/DatasmithImporter/Source/DatasmithImporter/Private/ActorFactoryDatasmithScene.cpp
--- a/Engine/Plugins/Enterprise/DatasmithImporter/Source/DatasmithImporter/Private/ActorFactoryDatasmithScene.cpp
+++ b/Engine/Plugins/Enterprise/DatasmithImporter/Source/DatasmithImporter/Private/ActorFactoryDatasmithScene.cpp
@@ -32,12 +32,13 @@ namespace UActorFactoryDatasmithSceneImpl
 	 *
 	 * @param DatasmithScene	The scene for which we will be spawning the actors.
 	 * @param RootActor			An optional existing ADatasmithSceneActor for which we will spawn/update the actors. It must refer to DatasmithScene.
+	 * @param InLevel			Optional level to spawn the actors, if not specified, will use Current Level
 	 * @param Transform			The transform to use if no RootActor was specified.
 	 * @param ObjectFlags		The flags to use when spawning new actors.
 	 *
 	 * @return					Returns either RootActor if it was supplied or the newly spawned ADatasmithSceneActor that refers to DatasmithScene.
 	 */
-	ADatasmithSceneActor* ImportActors( UDatasmithScene* DatasmithScene, ADatasmithSceneActor* RootActor, const FTransform& Transform, EObjectFlags ObjectFlags, bool bReimportDeletedActors )
+	ADatasmithSceneActor* ImportActors( UDatasmithScene* DatasmithScene, ADatasmithSceneActor* RootActor, ULevel* InLevel, const FTransform& Transform, EObjectFlags ObjectFlags, bool bReimportDeletedActors )
 	{
 		using namespace UE::DatasmithImporter;
 		// The root and scene must match in order to update an existing root from a scene
@@ -70,6 +71,11 @@ namespace UActorFactoryDatasmithSceneImpl
 		ImportContext.Options->BaseOptions = DatasmithScene->AssetImportData->BaseOptions;
 		ImportContext.Options->BaseOptions.SceneHandling = EDatasmithImportScene::CurrentLevel;
 
+		if (InLevel)
+		{
+			ImportContext.ActorsContext.ImportWorld = InLevel->GetWorld();
+		}
+
 		const bool bSilent = true;
 		TSharedPtr<FJsonObject> JsonOptions;
 		ImportContext.InitScene( DatasmithSceneElement.ToSharedRef() );
@@ -137,7 +143,7 @@ namespace UActorFactoryDatasmithSceneImpl
 
 void UActorFactoryDatasmithScene::SpawnRelatedActors( ADatasmithSceneActor* DatasmithSceneActor, bool bReimportDeletedActors )
 {
-	UActorFactoryDatasmithSceneImpl::ImportActors( DatasmithSceneActor->Scene, DatasmithSceneActor, DatasmithSceneActor->GetActorTransform(), DatasmithSceneActor->GetFlags(), bReimportDeletedActors );
+	UActorFactoryDatasmithSceneImpl::ImportActors( DatasmithSceneActor->Scene, DatasmithSceneActor, nullptr, DatasmithSceneActor->GetActorTransform(), DatasmithSceneActor->GetFlags(), bReimportDeletedActors );
 }
 
 UActorFactoryDatasmithScene::UActorFactoryDatasmithScene(const FObjectInitializer& ObjectInitializer)
@@ -185,7 +191,7 @@ AActor* UActorFactoryDatasmithScene::SpawnActor( UObject* Asset, ULevel* InLevel
 	}
 	else
 	{
-		ResultingActor = UActorFactoryDatasmithSceneImpl::ImportActors( DatasmithScene, nullptr, Transform, InSpawnParams.ObjectFlags, false );
+		ResultingActor = UActorFactoryDatasmithSceneImpl::ImportActors( DatasmithScene, nullptr, InLevel, Transform, InSpawnParams.ObjectFlags, false );
 	}
 
 	return ResultingActor;
