Fixed issues with Datasmith area light

#rnx

diff --git a/Engine/Plugins/Enterprise/DatasmithImporter/Source/DatasmithImporter/Private/DatasmithLightImporter.cpp b/Engine/Plugins/Enterprise/DatasmithImporter/Source/DatasmithImporter/Private/DatasmithLightImporter.cpp
--- a/Engine/Plugins/Enterprise/DatasmithImporter/Source/DatasmithImporter/Private/DatasmithLightImporter.cpp
+++ b/Engine/Plugins/Enterprise/DatasmithImporter/Source/DatasmithImporter/Private/DatasmithLightImporter.cpp
@@ -17,10 +17,13 @@
 #include "ObjectTemplates/DatasmithSceneComponentTemplate.h"
 #include "ObjectTemplates/DatasmithSkyLightComponentTemplate.h"
 
+#include "ActorFactories/ActorFactory.h"
 #include "AssetRegistry/AssetRegistryModule.h"
 #include "ComponentReregisterContext.h"
+#include "Editor.h"
 #include "EngineUtils.h"
 #include "Internationalization/Internationalization.h"
+#include "Logging/LogMacros.h"
 #include "ObjectTools.h"
 #include "RawMesh.h"
 
@@ -407,29 +410,23 @@ USceneComponent* FDatasmithLightImporter::ImportAreaLightComponent( const TShare
 {
 	TRACE_CPUPROFILER_EVENT_SCOPE(FDatasmithLightImporter::ImportAreaLightComponent);
 
-	USceneComponent* MainComponent = nullptr;
+	UE_LOG(LogDatasmithImport, Warning, TEXT("For %s, instead of a ChildActorComponent, a DatasmithAreaLightActor will be created instead."), AreaLightElement->GetLabel());
 
-	FSoftObjectPath LightShapeBlueprintRef = FSoftObjectPath( TEXT("/DatasmithContent/Datasmith/DatasmithArealight.DatasmithArealight") );
-	UBlueprint* LightShapeBlueprint = Cast< UBlueprint >( LightShapeBlueprintRef.TryLoad() );
+	AActor* AreaLightActor = nullptr;
 
-	if ( LightShapeBlueprint )
+	// This feature has stopped working in 4.27 :-(
+	// For now, just create an actor
+	// We will see whether it makes sense to address in future releases
+	if (AActor* RootActor = Cast<AActor>(Outer))
 	{
-		UChildActorComponent* ChildActorComponent = Cast< UChildActorComponent >( FDatasmithActorImporter::ImportSceneComponent( UChildActorComponent::StaticClass(), AreaLightElement, ImportContext, Outer, UniqueNameProvider ) );
-		ChildActorComponent->SetChildActorClass( TSubclassOf< AActor > ( LightShapeBlueprint->GeneratedClass ) );
-		ChildActorComponent->CreateChildActor();
-
-		ADatasmithAreaLightActor* LightShapeActor = Cast< ADatasmithAreaLightActor >( ChildActorComponent->GetChildActor() );
+		ImportContext.Hierarchy.Push(RootActor->GetRootComponent());
 
-		if ( LightShapeActor )
-		{
-			LightShapeActor->SetActorLabel( AreaLightElement->GetLabel() );
-			SetupAreaLightActor( AreaLightElement, ImportContext, LightShapeActor );
+		AreaLightActor = CreateAreaLightActor(AreaLightElement, ImportContext);
 
-			MainComponent = ChildActorComponent;
-		}
+		ImportContext.Hierarchy.Pop();
 	}
 
-	return MainComponent;
+	return AreaLightActor ? AreaLightActor->GetRootComponent() : nullptr;
 }
 
 ULightmassPortalComponent* FDatasmithLightImporter::ImportLightmassPortalComponent( const TSharedRef< IDatasmithLightmassPortalElement >& LightElement, FDatasmithImportContext& ImportContext, UObject* Outer, FDatasmithActorUniqueLabelProvider& UniqueNameProvider )
@@ -446,15 +443,51 @@ AActor* FDatasmithLightImporter::CreateAreaLightActor( const TSharedRef< IDatasm
 	FSoftObjectPath LightShapeBlueprintRef = FSoftObjectPath( TEXT("/DatasmithContent/Datasmith/DatasmithArealight.DatasmithArealight") );
 	UBlueprint* LightShapeBlueprint = Cast< UBlueprint >( LightShapeBlueprintRef.TryLoad() );
 
-	if ( !LightShapeBlueprint )
+	if (!ensure(LightShapeBlueprint))
+	{
+		return nullptr;
+	}
+
+	static TWeakObjectPtr<UActorFactory> ActorBPFactory;
+	if (!ActorBPFactory.IsValid())
+	{
+		FAssetData AssetData(LightShapeBlueprint);
+		FText UnusedErrorMessage;
+
+		const TArray<UActorFactory*>& ActorFactories = GEditor->ActorFactories;
+		for (int32 FactoryIdx = 0; FactoryIdx < ActorFactories.Num(); FactoryIdx++)
+		{
+			UActorFactory* ActorFactory = ActorFactories[FactoryIdx];
+
+			// Check if the actor can be created using this factory, making sure to check for an asset to be assigned from the selector
+			if (ActorFactory->CanCreateActorFrom(AssetData, UnusedErrorMessage))
+			{
+				ActorBPFactory = ActorFactory;
+				break;
+			}
+		}
+
+		if (!ActorBPFactory.IsValid())
+		{
+			return nullptr;
+		}
+	}
+
+	ULevel* CurrentLevel = ImportContext.ActorsContext.ImportWorld->GetCurrentLevel();
+	ADatasmithAreaLightActor* LightShapeActor = Cast< ADatasmithAreaLightActor >(ActorBPFactory->CreateActor(LightShapeBlueprint, CurrentLevel, FTransform::Identity));
+
+	if (!ensure(LightShapeActor))
 	{
 		return nullptr;
 	}
 
-	AActor* Actor = FDatasmithActorImporter::ImportActor( LightShapeBlueprint->GeneratedClass, AreaLightElement, ImportContext, ImportContext.Options->LightImportPolicy );
+	LightShapeActor->SetActorLabel(AreaLightElement->GetLabel());
+	FDatasmithActorImporter::SetupActorProperties(LightShapeActor, AreaLightElement, ImportContext);
+	
+	USceneComponent* Parent = ImportContext.Hierarchy.Num() > 0 ? ImportContext.Hierarchy.Top() : nullptr;
+	FDatasmithActorImporter::SetupSceneComponent(LightShapeActor->GetRootComponent(), AreaLightElement, Parent);
 
-	ADatasmithAreaLightActor* LightShapeActor = Cast< ADatasmithAreaLightActor >( Actor );
-	SetupAreaLightActor( AreaLightElement, ImportContext, LightShapeActor );
+	SetupAreaLightActor(AreaLightElement, ImportContext, LightShapeActor);
 
 	return LightShapeActor;
 }
