[Datasmith] Fixed imported textures settings messed up sometimes

- reset texture factory options to default values when those options are not always set up for different texture modes

diff --git a/Engine/Plugins/Enterprise/DatasmithImporter/Source/DatasmithImporter/Private/DatasmithTextureImporter.cpp b/Engine/Plugins/Enterprise/DatasmithImporter/Source/DatasmithImporter/Private/DatasmithTextureImporter.cpp
--- a/Engine/Plugins/Enterprise/DatasmithImporter/Source/DatasmithImporter/Private/DatasmithTextureImporter.cpp
+++ b/Engine/Plugins/Enterprise/DatasmithImporter/Source/DatasmithImporter/Private/DatasmithTextureImporter.cpp
@@ -212,6 +212,9 @@ UTexture* FDatasmithTextureImporter::CreateTexture(const TSharedPtr<IDatasmithTe
 		return nullptr;
 	}
 
+	// Reset options that are not always overwritten because single TextureFact instance is reused for every imported texture 
+	TextureFact->LODGroup = TEXTUREGROUP_World;
+	TextureFact->CompressionSettings = TC_Default;
 	TextureFact->bFlipNormalMapGreenChannel = false;
 
 	// Make sure to set the proper LODGroup as it's used to determine the CompressionSettings when using TEXTUREGROUP_WorldNormalMap
