[Datasmith] Geolocation metadata is now reimported

diff --git a/Engine/Plugins/Enterprise/DatasmithImporter/Source/DatasmithImporter/Private/DatasmithImporter.cpp b/Engine/Plugins/Enterprise/DatasmithImporter/Source/DatasmithImporter/Private/DatasmithImporter.cpp
--- a/Engine/Plugins/Enterprise/DatasmithImporter/Source/DatasmithImporter/Private/DatasmithImporter.cpp
+++ b/Engine/Plugins/Enterprise/DatasmithImporter/Source/DatasmithImporter/Private/DatasmithImporter.cpp
@@ -1385,6 +1385,29 @@ void FDatasmithImporter::FinalizeActors( FDatasmithImportContext& ImportContext,
 				}
 
 				FDatasmithImporterUtils::DeleteNonImportedDatasmithElementFromSceneActor( *ImportSceneActor, *DestinationSceneActor, ImportContext.ActorsContext.NonImportedDatasmithActors );
+
+
+				// Update destination actor's metadata
+				if (IInterface_AssetUserData* SourceAssetUserData = Cast<IInterface_AssetUserData>(ImportSceneActor->GetRootComponent()))
+				{
+					if (UDatasmithAssetUserData* SourceDatasmithUserData = SourceAssetUserData->GetAssetUserData<UDatasmithAssetUserData>())
+					{
+						USceneComponent* SceneComponent = DestinationSceneActor->GetRootComponent();
+
+						if (IInterface_AssetUserData* AssetUserData = Cast<IInterface_AssetUserData>(SceneComponent))
+						{
+							UDatasmithAssetUserData* DatasmithUserData = AssetUserData->GetAssetUserData<UDatasmithAssetUserData>();
+
+							if (!DatasmithUserData)
+							{
+								DatasmithUserData = NewObject<UDatasmithAssetUserData>(SceneComponent, NAME_None, RF_Public | RF_Transactional);
+								AssetUserData->AddAssetUserData(DatasmithUserData);
+							}
+
+							DatasmithUserData->MetaData.Append(SourceDatasmithUserData->MetaData);
+						}
+					}
+				}
 			}
 
 			// Add Actor info to the remap info
