Fixed craash

diff --git a/Engine/Plugins/Enterprise/DatasmithImporter/Source/DatasmithImporter/Private/ObjectElements/DatasmithUSceneElement.cpp b/Engine/Plugins/Enterprise/DatasmithImporter/Source/DatasmithImporter/Private/ObjectElements/DatasmithUSceneElement.cpp
--- a/Engine/Plugins/Enterprise/DatasmithImporter/Source/DatasmithImporter/Private/ObjectElements/DatasmithUSceneElement.cpp
+++ b/Engine/Plugins/Enterprise/DatasmithImporter/Source/DatasmithImporter/Private/ObjectElements/DatasmithUSceneElement.cpp
@@ -749,6 +749,11 @@ bool UDatasmithSceneElementBase::IsElementValid(const TWeakPtr<IDatasmithLevelVa
  */
 UDatasmithObjectElement* UDatasmithSceneElementBase::FindOrAddElement(const TSharedPtr<IDatasmithElement>& InElement)
 {
+	if (!InElement.IsValid())
+	{
+		return nullptr;
+	}
+
 	UDatasmithObjectElement* Object = nullptr;
 
 	if (InElement->IsA(EDatasmithElementType::Actor))
