Datasmith - Fixed the seams pattern indexation in the cloth importer, and reordered the imported positions axis to match Unreal's left handedness.

#rnx

diff --git a/Engine/Plugins/Enterprise/DatasmithImporter/Source/DatasmithImporter/Private/DatasmithImporter.cpp b/Engine/Plugins/Enterprise/DatasmithImporter/Source/DatasmithImporter/Private/DatasmithImporter.cpp
--- a/Engine/Plugins/Enterprise/DatasmithImporter/Source/DatasmithImporter/Private/DatasmithImporter.cpp
+++ b/Engine/Plugins/Enterprise/DatasmithImporter/Source/DatasmithImporter/Private/DatasmithImporter.cpp
@@ -493,16 +493,18 @@ void FDatasmithImporter::ImportClothes(FDatasmithImportContext& ImportContext)
 			continue;
 		}
 
+		using namespace UE::Chaos::ClothAsset;
+
 		TSharedPtr<FManagedArrayCollection> Collection = ClothAsset->GetClothCollection();
-		UE::Chaos::ClothAsset::FCollectionClothFacade Cloth(Collection);
+		FCollectionClothFacade Cloth(Collection);
 		Cloth.Reset(); // on reimport, asset is non-empty
-		UE::Chaos::ClothAsset::FCollectionClothLodFacade ClothLod = Cloth.GetNumLods() ? Cloth.GetLod(0) : Cloth.AddGetLod();
+		FCollectionClothLodFacade ClothLod = Cloth.GetNumLods() ? Cloth.GetLod(0) : Cloth.AddGetLod();
 
 		for (FDatasmithClothPattern& Pattern : DsCloth.Patterns)
 		{
 			if (Pattern.IsValid())
 			{
-				UE::Chaos::ClothAsset::FCollectionClothPatternFacade ClothPattern = ClothLod.AddGetPattern();
+				FCollectionClothPatternFacade ClothPattern = ClothLod.AddGetPattern();
 				ClothPattern.Initialize(Pattern.SimPosition, Pattern.SimRestPosition, Pattern.SimTriangleIndices);
 			}
 		}
@@ -511,19 +513,35 @@ void FDatasmithImporter::ImportClothes(FDatasmithImportContext& ImportContext)
 		TArrayView<FIntVector2> SeamPatterns = ClothLod.GetSeamPatterns();
 		TArrayView<TArray<FIntVector2>> SeamStitches = ClothLod.GetSeamStitches();
 
-		uint32 SeamIndex = 0;
+		int32 SeamIndex = 0;
 		for (const FDatasmithClothSewingInfo& SeamInfo : DsCloth.Sewing)
 		{
-			SeamPatterns[SeamIndex] = FIntVector2(SeamInfo.Seam0PanelIndex, SeamInfo.Seam1PanelIndex);
-			TArray<FIntVector2>& Stitches = SeamStitches[SeamIndex];
-			const uint32 StitchesCount = FMath::Min(SeamInfo.Seam0MeshIndices.Num(), SeamInfo.Seam1MeshIndices.Num());
-			Stitches.Reserve(StitchesCount);
-			for (uint32 StitchIndex = 0; StitchIndex < StitchesCount; ++StitchIndex)
+			const int32 SeamPattern0 = (int32)SeamInfo.Seam0PanelIndex;
+			const int32 SeamPattern1 = (int32)SeamInfo.Seam1PanelIndex;
+
+			if (SeamPattern0 >= 0 && SeamPattern0 < ClothLod.GetNumPatterns() &&
+				SeamPattern1 >= 0 && SeamPattern1 < ClothLod.GetNumPatterns())
 			{
-				Stitches.Add(FIntVector2(SeamInfo.Seam0MeshIndices[StitchIndex], SeamInfo.Seam1MeshIndices[StitchIndex]));
-			}
+				SeamPatterns[SeamIndex] = FIntVector2(SeamPattern0, SeamPattern1);
+
+				const FCollectionClothPatternConstFacade ClothPattern0 = ClothLod.GetPattern(SeamPattern0);
+				const FCollectionClothPatternConstFacade ClothPattern1 = ClothLod.GetPattern(SeamPattern1);
+
+				const int32 ClothPattern0VerticesOffset = ClothPattern0.GetSimVerticesOffset();
+				const int32 ClothPattern1VerticesOffset = ClothPattern1.GetSimVerticesOffset();
 
-			++SeamIndex;
+				TArray<FIntVector2>& Stitches = SeamStitches[SeamIndex];
+				const int32 StitchesCount = FMath::Min(SeamInfo.Seam0MeshIndices.Num(), SeamInfo.Seam1MeshIndices.Num());
+				Stitches.Reserve(StitchesCount);
+				for (int32 StitchIndex = 0; StitchIndex < StitchesCount; ++StitchIndex)
+				{
+					Stitches.Emplace(
+						(int32)SeamInfo.Seam0MeshIndices[StitchIndex] + ClothPattern0VerticesOffset,
+						(int32)SeamInfo.Seam1MeshIndices[StitchIndex] + ClothPattern1VerticesOffset);
+				}
+
+				++SeamIndex;
+			}
 		}
 
 		// Set a default skeleton
diff --git a/Engine/Plugins/Experimental/DatasmithCloTranslator/Source/DatasmithCloTranslator/Private/DatasmithCloTranslator.cpp b/Engine/Plugins/Experimental/DatasmithCloTranslator/Source/DatasmithCloTranslator/Private/DatasmithCloTranslator.cpp
--- a/Engine/Plugins/Experimental/DatasmithCloTranslator/Source/DatasmithCloTranslator/Private/DatasmithCloTranslator.cpp
+++ b/Engine/Plugins/Experimental/DatasmithCloTranslator/Source/DatasmithCloTranslator/Private/DatasmithCloTranslator.cpp
@@ -159,17 +159,17 @@ bool FDatasmithCloTranslator::LoadCloth(const TSharedRef<IDatasmithClothElement>
 		Panel.SimRestPosition = MoveTemp(CloPanel.Positions3D);
 		Panel.SimTriangleIndices = MoveTemp(CloPanel.TriangleIndices);
 
-		// Invert winding order
-		for (int32 TriIndex = 0; TriIndex < Panel.SimTriangleIndices.Num() / 3; ++TriIndex)
+		// Apply import scale and reverse winding order/Y axis for Unreal's left hand coordinate system
+		for (FVector3f& SimRestPosition : Panel.SimRestPosition)
 		{
-			Swap(Panel.SimTriangleIndices[3 * TriIndex + 1], Panel.SimTriangleIndices[3 * TriIndex + 2]);
+			SimRestPosition = FVector3f(-SimRestPosition.Y, -SimRestPosition.X, SimRestPosition.Z) * ImportOptions.Scale;
 		}
 
 		if (ImportOptions.Scale != 1.0f)
 		{
-			for (auto& Vertex : Panel.SimRestPosition)
+			for (FVector2f& SimPosition : Panel.SimPosition)
 			{
-				Vertex *= ImportOptions.Scale;
+				SimPosition *= ImportOptions.Scale;
 			}
 		}
 
@@ -509,8 +509,8 @@ bool FDatasmithCloTranslator::ParseFromJson()
 
 			FCloSewingInfo& SewInfo = CloCloth.SewingInfos.AddDefaulted_GetRef();
 			SewInfo.Index = SewingInfoJsonObject->GetIntegerField(TEXT("index"));
-			SewInfo.Seam0PanelIndex = SewingInfoJsonObject->GetIntegerField(TEXT("seam0PanelIndex"));
-			SewInfo.Seam1PanelIndex = SewingInfoJsonObject->GetIntegerField(TEXT("seam1PanelIndex"));
+			SewInfo.Seam0PanelIndex = SewingInfoJsonObject->GetIntegerField(TEXT("seam0PatternIndex"));
+			SewInfo.Seam1PanelIndex = SewingInfoJsonObject->GetIntegerField(TEXT("seam1PatternIndex"));
 			AppendIntList(SewingInfoJsonObject->GetArrayField(TEXT("seam0MeshIndexList")), SewInfo.Seam0MeshIndices);
 			AppendIntList(SewingInfoJsonObject->GetArrayField(TEXT("seam1MeshIndexList")), SewInfo.Seam1MeshIndices);
 		}
