Fixed deorecated string hashing function in a .ContainsByHash(Hash, CharPtr) call FDatasmithImporterUtils::GetOrderedListOfMaterialsReferencedByMaterials.

diff --git a/Engine/Plugins/Enterprise/DatasmithImporter/Source/DatasmithImporter/Private/Utility/DatasmithImporterUtils.cpp b/Engine/Plugins/Enterprise/DatasmithImporter/Source/DatasmithImporter/Private/Utility/DatasmithImporterUtils.cpp
--- a/Engine/Plugins/Enterprise/DatasmithImporter/Source/DatasmithImporter/Private/Utility/DatasmithImporterUtils.cpp
+++ b/Engine/Plugins/Enterprise/DatasmithImporter/Source/DatasmithImporter/Private/Utility/DatasmithImporterUtils.cpp
@@ -1145,7 +1145,7 @@ void FDatasmithImporterUtils::FillSceneElement(TSharedPtr<IDatasmithScene>& Scen
 TArray<FDatasmithImporterUtils::FFunctionAndMaterialsThatUseIt> FDatasmithImporterUtils::GetOrderedListOfMaterialsReferencedByMaterials(TSharedPtr< IDatasmithScene >& SceneElement)
 {
 	//This map is used to keep track of which materials are referencing which
-	TMap<FString, TSet<FString>> MaterialToFunctionNameMap;
+	TMap<FString, TSet<FString, TStringPointerSetKeyFuncs_DEPRECATED<FString>>> MaterialToFunctionNameMap;
 	//Mapping materials to their names for easy access.
 	TMap<FString, TSharedPtr<IDatasmithUEPbrMaterialElement>> MaterialNameMap;
 
@@ -1175,7 +1175,7 @@ TArray<FDatasmithImporterUtils::FFunctionAndMaterialsThatUseIt> FDatasmithImport
 				if (FPaths::IsRelative(FunctionPathName))
 				{
 					uint32 FunctionPathNameHash = GetTypeHash(FunctionPathName);
-					if (TSet<FString>* Values = MaterialToFunctionNameMap.FindByHash(FunctionPathNameHash, FunctionPathName))
+					if (TSet<FString, TStringPointerSetKeyFuncs_DEPRECATED<FString>>* Values = MaterialToFunctionNameMap.FindByHash(FunctionPathNameHash, FunctionPathName))
 					{
 						check(!Values->ContainsByHash(BaseMaterialHash, BaseMaterialElement->GetName())); //Can't have inter-dependencies
 					}
