Fixed compile error.

diff --git a/Engine/Plugins/Enterprise/DatasmithImporter/Source/DatasmithImporter/Private/Utility/DatasmithImporterUtils.cpp b/Engine/Plugins/Enterprise/DatasmithImporter/Source/DatasmithImporter/Private/Utility/DatasmithImporterUtils.cpp
--- a/Engine/Plugins/Enterprise/DatasmithImporter/Source/DatasmithImporter/Private/Utility/DatasmithImporterUtils.cpp
+++ b/Engine/Plugins/Enterprise/DatasmithImporter/Source/DatasmithImporter/Private/Utility/DatasmithImporterUtils.cpp
@@ -1213,7 +1213,7 @@ TArray<FDatasmithImporterUtils::FFunctionAndMaterialsThatUseIt> FDatasmithImport
 				const TSharedPtr<IDatasmithUEPbrMaterialElement>* NextElement = MaterialNameMap.FindByHash(DependencyHash, Dependency);
 				if (NextElement)
 				{
-					if (const TSet<FString>* NextElementDependencies = MaterialToFunctionNameMap.FindByHash(DependencyHash, Dependency))
+					if (const TSet<FString, TStringPointerSetKeyFuncs_DEPRECATED<FString>>* NextElementDependencies = MaterialToFunctionNameMap.FindByHash(DependencyHash, Dependency))
 					{
 						DepthFirstSearch(*NextElement, *NextElementDependencies, CurrentElement);
 					}
@@ -1245,7 +1245,7 @@ TArray<FDatasmithImporterUtils::FFunctionAndMaterialsThatUseIt> FDatasmithImport
 				const TSharedPtr<IDatasmithUEPbrMaterialElement>* CurrentDependency = MaterialNameMap.FindByHash(DependencyHash, Dependency);
 				if (CurrentDependency)
 				{
-					if (const TSet<FString>* DependenciesOfDependency = MaterialToFunctionNameMap.FindByHash(DependencyHash, Dependency))
+					if (const TSet<FString, TStringPointerSetKeyFuncs_DEPRECATED<FString>>* DependenciesOfDependency = MaterialToFunctionNameMap.FindByHash(DependencyHash, Dependency))
 					{
 						DepthFirstSearch(*CurrentDependency, *DependenciesOfDependency, CurrentElement);
 					}
