More compile fixes.

diff --git a/Engine/Plugins/Enterprise/DatasmithImporter/Source/DatasmithImporter/Private/Utility/DatasmithImporterUtils.cpp b/Engine/Plugins/Enterprise/DatasmithImporter/Source/DatasmithImporter/Private/Utility/DatasmithImporterUtils.cpp
--- a/Engine/Plugins/Enterprise/DatasmithImporter/Source/DatasmithImporter/Private/Utility/DatasmithImporterUtils.cpp
+++ b/Engine/Plugins/Enterprise/DatasmithImporter/Source/DatasmithImporter/Private/Utility/DatasmithImporterUtils.cpp
@@ -1193,9 +1193,9 @@ TArray<FDatasmithImporterUtils::FFunctionAndMaterialsThatUseIt> FDatasmithImport
 	VisitedAndInsertionIndex.Reserve(MaterialToFunctionNameMap.Num());
 
 	// Topological sort while also noting who was the direct dependents of a node also. Base on a DFS
-	TFunction<void (const TSharedPtr<IDatasmithUEPbrMaterialElement>&, const TSet<FString>&, const TSharedPtr<IDatasmithUEPbrMaterialElement>&)> DepthFirstSearch;
+	TFunction<void (const TSharedPtr<IDatasmithUEPbrMaterialElement>&, const TSet<FString, TStringPointerSetKeyFuncs_DEPRECATED<FString>>&, const TSharedPtr<IDatasmithUEPbrMaterialElement>&)> DepthFirstSearch;
 
-	DepthFirstSearch = [&MaterialToFunctionNameMap, &MaterialNameMap, &ReferencedMaterialAndTheirReferencer, &VisitedAndInsertionIndex, &DepthFirstSearch] (const TSharedPtr<IDatasmithUEPbrMaterialElement>& CurrentElement, const TSet<FString>& Dependencies, const TSharedPtr<IDatasmithUEPbrMaterialElement>& Parent)
+	DepthFirstSearch = [&MaterialToFunctionNameMap, &MaterialNameMap, &ReferencedMaterialAndTheirReferencer, &VisitedAndInsertionIndex, &DepthFirstSearch] (const TSharedPtr<IDatasmithUEPbrMaterialElement>& CurrentElement, const TSet<FString, TStringPointerSetKeyFuncs_DEPRECATED<FString>>& Dependencies, const TSharedPtr<IDatasmithUEPbrMaterialElement>& Parent)
 		{
 			uint32 CurrentHash = GetTypeHash(CurrentElement);
 			// Check if visited
@@ -1219,7 +1219,7 @@ TArray<FDatasmithImporterUtils::FFunctionAndMaterialsThatUseIt> FDatasmithImport
 					}
 					else
 					{
-						DepthFirstSearch(*NextElement, TSet<FString>(), CurrentElement);
+						DepthFirstSearch(*NextElement, TSet<FString, TStringPointerSetKeyFuncs_DEPRECATED<FString>>(), CurrentElement);
 					}
 				}
 			}
@@ -1232,7 +1232,7 @@ TArray<FDatasmithImporterUtils::FFunctionAndMaterialsThatUseIt> FDatasmithImport
 			ReferencedMaterialAndTheirReferencer.Emplace(CurrentElement, MoveTemp(Parents));
 		};
 
-	for (const TPair<FString, TSet<FString>>& Pair : MaterialToFunctionNameMap)
+	for (const TPair<FString, TSet<FString, TStringPointerSetKeyFuncs_DEPRECATED<FString>>>& Pair : MaterialToFunctionNameMap)
 	{
 		const TSharedPtr<IDatasmithUEPbrMaterialElement>& CurrentElement = MaterialNameMap.FindChecked(Pair.Key);
 		// Only if not yet visited
@@ -1251,7 +1251,7 @@ TArray<FDatasmithImporterUtils::FFunctionAndMaterialsThatUseIt> FDatasmithImport
 					}
 					else
 					{
-						DepthFirstSearch(*CurrentDependency, TSet<FString>(), CurrentElement);
+						DepthFirstSearch(*CurrentDependency, TSet<FString, TStringPointerSetKeyFuncs_DEPRECATED<FString>>(), CurrentElement);
 					}
 				}
 			}
