Fixed missing addition of certains types of actors to the import context
Root cause: Some methods creating actors, i.e. light actors, are not calling FDatasmithActorImporter::ImportActor

#rnx

diff --git a/Engine/Plugins/Enterprise/DatasmithImporter/Source/DatasmithImporter/Private/DatasmithActorImporter.cpp b/Engine/Plugins/Enterprise/DatasmithImporter/Source/DatasmithImporter/Private/DatasmithActorImporter.cpp
--- a/Engine/Plugins/Enterprise/DatasmithImporter/Source/DatasmithImporter/Private/DatasmithActorImporter.cpp
+++ b/Engine/Plugins/Enterprise/DatasmithImporter/Source/DatasmithImporter/Private/DatasmithActorImporter.cpp
@@ -136,8 +136,6 @@ AActor* FDatasmithActorImporter::ImportActor( UClass* ActorClass, const TSharedR
 		SetupActorProperties( ImportedActor, ActorElement, ImportContext ) ;
 		SetupSceneComponent( ImportedActor->GetRootComponent(), ActorElement, ImportContext.Hierarchy.Num() > 0 ? ImportContext.Hierarchy.Top() : nullptr );
 
-		ImportContext.AddImportedActor( ImportedActor );
-
 		if ( DatasmithSceneActor )
 		{
 			DatasmithSceneActor->RelatedActors.FindOrAdd( ActorElement->GetName() ) = ImportedActor;
diff --git a/Engine/Plugins/Enterprise/DatasmithImporter/Source/DatasmithImporter/Private/DatasmithImporter.cpp b/Engine/Plugins/Enterprise/DatasmithImporter/Source/DatasmithImporter/Private/DatasmithImporter.cpp
--- a/Engine/Plugins/Enterprise/DatasmithImporter/Source/DatasmithImporter/Private/DatasmithImporter.cpp
+++ b/Engine/Plugins/Enterprise/DatasmithImporter/Source/DatasmithImporter/Private/DatasmithImporter.cpp
@@ -1189,11 +1189,12 @@ AActor* FDatasmithImporter::ImportActor( FDatasmithImportContext& ImportContext,
 		ImportedActor = FDatasmithActorImporter::ImportBaseActor( ImportContext, ActorElement );
 	}
 
-
 	if ( ImportedActor ) // It's possible that we didn't import an actor (ie: the user doesn't want to import the cameras), in that case, we'll skip it in the hierarchy
 	{
 		ImportContext.Hierarchy.Push( ImportedActor->GetRootComponent() );
 		ImportMetaDataForObject(ImportContext, ActorElement, ImportedActor);
+
+		ImportContext.AddImportedActor(ImportedActor);
 	}
 	else
 	{
