Chaos Cloth Asset - Severed the Datasmith Importer dependencies to the Chaos Cloth Asset, and allowed for more cloth object types to also use the importer.

#rnx

diff --git a/Engine/Plugins/Enterprise/DatasmithImporter/DatasmithImporter.uplugin b/Engine/Plugins/Enterprise/DatasmithImporter/DatasmithImporter.uplugin
--- a/Engine/Plugins/Enterprise/DatasmithImporter/DatasmithImporter.uplugin
+++ b/Engine/Plugins/Enterprise/DatasmithImporter/DatasmithImporter.uplugin
@@ -90,10 +90,6 @@
 		{
 			"Name": "UdpMessaging",
 			"Enabled": true
-		},
-		{
-			"Name": "ChaosClothAsset",
-			"Enabled": true
 		}
 	]
 }
\ No newline at end of file
diff --git a/Engine/Plugins/Enterprise/DatasmithImporter/Source/DatasmithImporter/DatasmithImporter.Build.cs b/Engine/Plugins/Enterprise/DatasmithImporter/Source/DatasmithImporter/DatasmithImporter.Build.cs
--- a/Engine/Plugins/Enterprise/DatasmithImporter/Source/DatasmithImporter/DatasmithImporter.Build.cs
+++ b/Engine/Plugins/Enterprise/DatasmithImporter/Source/DatasmithImporter/DatasmithImporter.Build.cs
@@ -12,9 +12,6 @@ namespace UnrealBuildTool.Rules
 				{
 					"Analytics",
 					"ApplicationCore",
-					"Chaos",
-					"ChaosClothAsset",
-					"ChaosClothAssetEngine",
 					"CinematicCamera",
 					"ContentBrowser",
 					"Core",
diff --git a/Engine/Plugins/Enterprise/DatasmithImporter/Source/DatasmithImporter/Private/DatasmithActorImporter.cpp b/Engine/Plugins/Enterprise/DatasmithImporter/Source/DatasmithImporter/Private/DatasmithActorImporter.cpp
--- a/Engine/Plugins/Enterprise/DatasmithImporter/Source/DatasmithImporter/Private/DatasmithActorImporter.cpp
+++ b/Engine/Plugins/Enterprise/DatasmithImporter/Source/DatasmithImporter/Private/DatasmithActorImporter.cpp
@@ -3,6 +3,7 @@
 #include "DatasmithActorImporter.h"
 
 #include "DatasmithCameraImporter.h"
+#include "DatasmithCloth.h"
 #include "DatasmithImportContext.h"
 #include "DatasmithImporterModule.h"
 #include "DatasmithImportOptions.h"
@@ -23,8 +24,6 @@
 #include "Utility/DatasmithMeshHelper.h"
 
 #include "ActorFactories/ActorFactoryDeferredDecal.h"
-#include "ChaosClothAsset/ClothAsset.h"
-#include "ChaosClothAsset/ClothComponent.h"
 #include "CineCameraActor.h"
 #include "CineCameraComponent.h"
 #include "Components/DecalComponent.h"
@@ -35,6 +34,7 @@
 #include "Engine/DecalActor.h"
 #include "Engine/StaticMeshActor.h"
 #include "Engine/World.h"
+#include "Features/IModularFeatures.h"
 #include "Materials/Material.h"
 #include "Materials/MaterialInterface.h"
 #include "Misc/Paths.h"
@@ -317,19 +317,24 @@ AActor* FDatasmithActorImporter::ImportClothActor(FDatasmithImportContext& Impor
 			}
 		}
 
-		if (UChaosClothComponent* ClothComponent = NewObject<UChaosClothComponent>(ImportedActor))
+		// Locate a cloth factory provider
+		const TArray<IDatasmithClothFactoryClassesProvider*> ClothFactoryClassesProviders = IModularFeatures::Get().GetModularFeatureImplementations<IDatasmithClothFactoryClassesProvider>(IDatasmithClothFactoryClassesProvider::FeatureName);
+		const IDatasmithClothFactoryClassesProvider* const ClothFactoryClassesProvider = ClothFactoryClassesProviders.Num() ? ClothFactoryClassesProviders[0] : nullptr;
+		const UDatasmithClothComponentFactory* const ClothComponentFactory = ClothFactoryClassesProvider ? ClothFactoryClassesProvider->GetClothComponentFactoryClass().GetDefaultObject() : nullptr;
+		
+		// Create a new cloth component
+		if (USceneComponent* const ClothComponent = ClothComponentFactory ? ClothComponentFactory->CreateClothComponent(ImportedActor) : nullptr)
 		{
-			ClothComponent->SetClothAsset(Cast<UChaosClothAsset>(ClothAsset));
+			ClothComponentFactory->InitializeClothComponent(ClothComponent, ClothAsset, ImportedActor->GetRootComponent());
 
 			ImportedActor->AddInstanceComponent(ClothComponent);
-			ClothComponent->SetupAttachment(ImportedActor->GetRootComponent());
 
 			ClothComponent->bVisualizeComponent = true;
 			ClothComponent->RegisterComponent();
 		}
 		else
 		{
-			ImportContext.LogError(FText::Format(LOCTEXT("ClothComponentCreationFailure", "Cannot create cloth component for asset {0} reauired by actor {1}."), FText::FromString(ClothName), FText::FromString(ClothActorElement->GetLabel())));
+			ImportContext.LogError(FText::Format(LOCTEXT("ClothComponentCreationFailure", "Cannot create cloth component for asset {0} required by actor {1}."), FText::FromString(ClothName), FText::FromString(ClothActorElement->GetLabel())));
 		}
 
 		if (ClothAsset == nullptr)
diff --git a/Engine/Plugins/Enterprise/DatasmithImporter/Source/DatasmithImporter/Private/DatasmithImporter.cpp b/Engine/Plugins/Enterprise/DatasmithImporter/Source/DatasmithImporter/Private/DatasmithImporter.cpp
--- a/Engine/Plugins/Enterprise/DatasmithImporter/Source/DatasmithImporter/Private/DatasmithImporter.cpp
+++ b/Engine/Plugins/Enterprise/DatasmithImporter/Source/DatasmithImporter/Private/DatasmithImporter.cpp
@@ -48,6 +48,7 @@
 #include "Engine/SkinnedAssetCommon.h"
 #include "Engine/StaticMesh.h"
 #include "Engine/StaticMeshActor.h"
+#include "Features/IModularFeatures.h"
 #include "HAL/FileManager.h"
 #include "IAssetTools.h"
 #include "ISourceControlModule.h"
@@ -81,11 +82,6 @@
 #include "UObject/UObjectHash.h"
 #include "UnrealEdGlobals.h"
 
-#include "Animation/Skeleton.h"
-#include "ChaosClothAsset/ClothAsset.h"
-#include "ChaosClothAsset/CollectionClothFacade.h"
-#include "GeometryCollection/ManagedArrayCollection.h"
-
 extern UNREALED_API UEditorEngine* GEditor;
 
 #define LOCTEXT_NAMESPACE "DatasmithImporter"
@@ -390,6 +386,24 @@ void FDatasmithImporter::ImportClothes(FDatasmithImportContext& ImportContext)
 
 	TRACE_CPUPROFILER_EVENT_SCOPE(FDatasmithImporter::ImportClothes);
 
+	// Locate a cloth asset provider
+	const TArray<IDatasmithClothFactoryClassesProvider*> ClothFactoryClassesProviders = IModularFeatures::Get().GetModularFeatureImplementations<IDatasmithClothFactoryClassesProvider>(IDatasmithClothFactoryClassesProvider::FeatureName);
+	if (!ClothFactoryClassesProviders.Num() || !ClothFactoryClassesProviders[0])
+	{
+		UE_LOG(LogDatasmithImport, Warning, TEXT("No cloth factory class provider found, cloth import skipped."));
+		return;
+	}
+	const UDatasmithClothAssetFactory* const ClothAssetFactory = ClothFactoryClassesProviders[0]->GetClothAssetFactoryClass().GetDefaultObject();
+	if (!ClothAssetFactory)
+	{
+		UE_LOG(LogDatasmithImport, Warning, TEXT("Invalid (null) cloth asset factory, cloth import skipped."));
+		return;
+	}
+	if (ClothFactoryClassesProviders.Num() > 1)
+	{
+		UE_LOG(LogDatasmithImport, Warning, TEXT("Multiple cloth asset providers found, only the cloth asset type from provider %s will be imported."), *ClothFactoryClassesProviders[0]->GetName().ToString());
+	}
+
 	TUniquePtr<FScopedSlowTask> ProgressPtr;
 	if (ImportContext.FeedbackContext)
 	{
@@ -437,8 +451,6 @@ void FDatasmithImporter::ImportClothes(FDatasmithImportContext& ImportContext)
 		UPackage* ImportPackage = ImportContext.AssetsContext.StaticMeshesImportPackage.Get();
 		EObjectFlags ObjectFlags = ImportContext.ObjectFlags & ~RF_Public; // not RF_Public yet, the publicized asset will be.
 
-		UChaosClothAsset* ClothAsset = nullptr;
-
 		// TODO: Use FCollectionPropertyFacade instead, see \Engine\Source\Runtime\Experimental\Chaos\Public\Chaos\CollectionPropertyFacade.h
 		//TArray<UObject*> ClothPresetAssets;
 
@@ -456,22 +468,22 @@ void FDatasmithImporter::ImportClothes(FDatasmithImportContext& ImportContext)
 		//	}
 		//}
 
-		UObject* ExistingClothObj = nullptr;
+		UObject* ExistingCloth = nullptr;
 		if (ImportContext.SceneAsset)
 		{
 			if (TSoftObjectPtr<UObject>* ExistingAssetPtr = ImportContext.SceneAsset->Clothes.Find(ClothElement->GetName()))
 			{
-				ExistingClothObj = ExistingAssetPtr->LoadSynchronous();
+				ExistingCloth = ExistingAssetPtr->LoadSynchronous();
 			}
 		}
-		UChaosClothAsset* ExistingCloth = Cast<UChaosClothAsset>(ExistingClothObj);
+		UObject* ClothAsset = nullptr;
 		if (ExistingCloth)
 		{
 			if (ExistingCloth->GetOuter() != ImportPackage)
 			{
 				// Temporary flag to skip PostLoad during DuplicateObject
 				ExistingCloth->SetFlags(RF_ArchetypeObject);
-				ClothAsset = DuplicateObject<UChaosClothAsset>(ExistingCloth, ImportPackage, *ClothName);
+				ClothAsset = ClothAssetFactory->DuplicateClothAsset(ExistingCloth, ImportPackage, *ClothName);
 				ExistingCloth->ClearFlags(RF_ArchetypeObject);
 				ClothAsset->ClearFlags(RF_ArchetypeObject);
 
@@ -486,7 +498,7 @@ void FDatasmithImporter::ImportClothes(FDatasmithImportContext& ImportContext)
 		}
 		else
 		{
-			ClothAsset = NewObject<UChaosClothAsset>(ImportPackage, *ClothName, ObjectFlags);
+			ClothAsset = ClothAssetFactory->CreateClothAsset(ImportPackage, *ClothName, ObjectFlags);
 		}
 
 		if (!ensure(ClothAsset))
@@ -494,56 +506,7 @@ void FDatasmithImporter::ImportClothes(FDatasmithImportContext& ImportContext)
 			continue;
 		}
 
-		using namespace UE::Chaos::ClothAsset;
-
-		TArray<TSharedRef<FManagedArrayCollection>>& Collections = ClothAsset->GetClothCollections();
-		Collections.Reset(1);
-		FCollectionClothFacade Cloth(Collections.Emplace_GetRef(MakeShared<FManagedArrayCollection>()));
-		Cloth.DefineSchema();
-
-		for (FDatasmithClothPattern& Pattern : DsCloth.Patterns)
-		{
-			if (Pattern.IsValid())
-			{
-				FCollectionClothSimPatternFacade ClothPattern = Cloth.AddGetSimPattern();
-				ClothPattern.Initialize(Pattern.SimPosition, Pattern.SimRestPosition, Pattern.SimTriangleIndices);
-			}
-		}
-
-		for (const FDatasmithClothSewingInfo& SeamInfo : DsCloth.Sewing)
-		{
-			const int32 SeamPattern0 = (int32)SeamInfo.Seam0PanelIndex;
-			const int32 SeamPattern1 = (int32)SeamInfo.Seam1PanelIndex;
-
-			if (SeamPattern0 >= 0 && SeamPattern0 < Cloth.GetNumSimPatterns() &&
-				SeamPattern1 >= 0 && SeamPattern1 < Cloth.GetNumSimPatterns())
-			{
-				const FCollectionClothSimPatternConstFacade ClothPattern0 = Cloth.GetSimPattern(SeamPattern0);
-				const FCollectionClothSimPatternConstFacade ClothPattern1 = Cloth.GetSimPattern(SeamPattern1);
-
-				const int32 ClothPattern0VerticesOffset = ClothPattern0.GetSimVertices2DOffset();
-				const int32 ClothPattern1VerticesOffset = ClothPattern1.GetSimVertices2DOffset();
-
-				TArray<FIntVector2> Stitches;
-				const uint32 StitchesCount = FMath::Min(SeamInfo.Seam0MeshIndices.Num(), SeamInfo.Seam1MeshIndices.Num());
-				Stitches.Reserve(StitchesCount);
-				for (uint32 StitchIndex = 0; StitchIndex < StitchesCount; ++StitchIndex)
-				{
-					Stitches.Emplace(
-						(int32)SeamInfo.Seam0MeshIndices[StitchIndex] + ClothPattern0VerticesOffset, 
-						(int32)SeamInfo.Seam1MeshIndices[StitchIndex] + ClothPattern1VerticesOffset);
-				}
-
-				FCollectionClothSeamFacade SeamFacade = Cloth.AddGetSeam();
-				SeamFacade.Initialize(Stitches);
-			}
-		}
-
-		// Set the render mesh to duplicate the sim mesh
-		ClothAsset->CopySimMeshToRenderMesh();
-
-		// Set a default skeleton and rebuild the asset
-		ClothAsset->SetReferenceSkeleton(nullptr);  // This creates a default reference skeleton, redo the bindings, and rebuild the asset
+		ClothAssetFactory->InitializeClothAsset(ClothAsset, DsCloth);
 
 		ImportContext.ImportedClothes.Add(ClothElement, ClothAsset);
 
diff --git a/Engine/Plugins/Experimental/ChaosClothAssetEditor/Source/ChaosClothAssetDataflowNodes/ChaosClothAssetDataflowNodes.Build.cs b/Engine/Plugins/Experimental/ChaosClothAssetEditor/Source/ChaosClothAssetDataflowNodes/ChaosClothAssetDataflowNodes.Build.cs
--- a/Engine/Plugins/Experimental/ChaosClothAssetEditor/Source/ChaosClothAssetDataflowNodes/ChaosClothAssetDataflowNodes.Build.cs
+++ b/Engine/Plugins/Experimental/ChaosClothAssetEditor/Source/ChaosClothAssetDataflowNodes/ChaosClothAssetDataflowNodes.Build.cs
@@ -27,6 +27,7 @@ public class ChaosClothAssetDataflowNodes : ModuleRules
 				"DataflowCore",
 				"DataflowEditor",
 				"DataflowEngine",
+				"DatasmithCore",
 				"DatasmithImporter",
 				"DetailCustomizations",
 				"DynamicMesh",
diff --git a/Engine/Plugins/Experimental/ChaosClothAssetEditor/Source/ChaosClothAssetDataflowNodes/Private/ChaosClothAsset/ChaosClothAssetDataflowNodesModule.cpp b/Engine/Plugins/Experimental/ChaosClothAssetEditor/Source/ChaosClothAssetDataflowNodes/Private/ChaosClothAsset/ChaosClothAssetDataflowNodesModule.cpp
--- a/Engine/Plugins/Experimental/ChaosClothAssetEditor/Source/ChaosClothAssetDataflowNodes/Private/ChaosClothAsset/ChaosClothAssetDataflowNodesModule.cpp
+++ b/Engine/Plugins/Experimental/ChaosClothAssetEditor/Source/ChaosClothAssetDataflowNodes/Private/ChaosClothAsset/ChaosClothAssetDataflowNodesModule.cpp
@@ -111,6 +111,9 @@ void FChaosClothAssetDataflowNodesModule::StartupModule()
 		PropertyModule->RegisterCustomPropertyTypeLayout("ChaosClothAssetWeightedValueNonAnimatable", FOnGetPropertyTypeCustomizationInstance::CreateStatic(&FWeightedValueCustomization::MakeInstance));
 		PropertyModule->RegisterCustomPropertyTypeLayout("ChaosClothAssetWeightedValueNonAnimatableNoLowHighRange", FOnGetPropertyTypeCustomizationInstance::CreateStatic(&FWeightedValueCustomization::MakeInstance));
 	}
+
+	// Register modular features
+	FChaosClothAssetDatasmithImportNode::RegisterModularFeature();
 }
 
 void FChaosClothAssetDataflowNodesModule::ShutdownModule()
@@ -122,6 +125,9 @@ void FChaosClothAssetDataflowNodesModule::ShutdownModule()
 		PropertyModule->UnregisterCustomPropertyTypeLayout("ChaosClothAssetWeightedValueNonAnimatable");
 		PropertyModule->UnregisterCustomPropertyTypeLayout("ChaosClothAssetWeightedValueNonAnimatableNoLowHighRange");
 	}
+
+	// Unregister modular features
+	FChaosClothAssetDatasmithImportNode::UnregisterModularFeature();
 }
 
 #undef LOCTEXT_NAMESPACE
diff --git a/Engine/Plugins/Experimental/ChaosClothAssetEditor/Source/ChaosClothAssetDataflowNodes/Private/ChaosClothAsset/DatasmithImportNode.cpp b/Engine/Plugins/Experimental/ChaosClothAssetEditor/Source/ChaosClothAssetDataflowNodes/Private/ChaosClothAsset/DatasmithImportNode.cpp
--- a/Engine/Plugins/Experimental/ChaosClothAssetEditor/Source/ChaosClothAssetDataflowNodes/Private/ChaosClothAsset/DatasmithImportNode.cpp
+++ b/Engine/Plugins/Experimental/ChaosClothAssetEditor/Source/ChaosClothAssetDataflowNodes/Private/ChaosClothAsset/DatasmithImportNode.cpp
@@ -2,8 +2,10 @@
 
 #include "ChaosClothAsset/DatasmithImportNode.h"
 #include "ChaosClothAsset/ClothAsset.h"
+#include "ChaosClothAsset/ClothComponent.h"
 #include "ChaosClothAsset/CollectionClothFacade.h"
 #include "Dataflow/DataflowInputOutput.h"
+#include "Features/IModularFeatures.h"
 #include "DatasmithImportContext.h"
 #include "DatasmithImportFactory.h"
 #include "ExternalSource.h"
@@ -14,6 +16,131 @@
 
 #define LOCTEXT_NAMESPACE "ChaosClothAssetDatasmithImportNode"
 
+UChaosClothAssetDatasmithClothAssetFactory::UChaosClothAssetDatasmithClothAssetFactory() = default;
+UChaosClothAssetDatasmithClothAssetFactory::~UChaosClothAssetDatasmithClothAssetFactory() = default;
+
+UObject* UChaosClothAssetDatasmithClothAssetFactory::CreateClothAsset(UObject* Outer, const FName& Name, EObjectFlags Flags) const
+{
+	return CastChecked<UObject>(NewObject<UChaosClothAsset>(Outer, Name, Flags));
+}
+
+UObject* UChaosClothAssetDatasmithClothAssetFactory::DuplicateClothAsset(UObject* ClothAsset, UObject* Outer, const FName& Name) const
+{
+	return CastChecked<UObject>(DuplicateObject<UChaosClothAsset>(CastChecked<UChaosClothAsset>(ClothAsset), Outer, Name));
+}
+
+void UChaosClothAssetDatasmithClothAssetFactory::InitializeClothAsset(UObject* ClothAsset, const FDatasmithCloth& DatasmithCloth) const
+{
+	using namespace UE::Chaos::ClothAsset;
+
+	UChaosClothAsset* const ChaosClothAsset = CastChecked<UChaosClothAsset>(ClothAsset);
+
+	TArray<TSharedRef<FManagedArrayCollection>>& Collections = ChaosClothAsset->GetClothCollections();
+	Collections.Reset(1);
+	FCollectionClothFacade Cloth(Collections.Emplace_GetRef(MakeShared<FManagedArrayCollection>()));
+	Cloth.DefineSchema();
+
+	for (const FDatasmithClothPattern& Pattern : DatasmithCloth.Patterns)
+	{
+		if (Pattern.IsValid())
+		{
+			FCollectionClothSimPatternFacade ClothPattern = Cloth.AddGetSimPattern();
+			ClothPattern.Initialize(Pattern.SimPosition, Pattern.SimRestPosition, Pattern.SimTriangleIndices);
+		}
+	}
+
+	for (const FDatasmithClothSewingInfo& SeamInfo : DatasmithCloth.Sewing)
+	{
+		const int32 SeamPattern0 = (int32)SeamInfo.Seam0PanelIndex;
+		const int32 SeamPattern1 = (int32)SeamInfo.Seam1PanelIndex;
+
+		if (SeamPattern0 >= 0 && SeamPattern0 < Cloth.GetNumSimPatterns() &&
+			SeamPattern1 >= 0 && SeamPattern1 < Cloth.GetNumSimPatterns())
+		{
+			const FCollectionClothSimPatternConstFacade ClothPattern0 = Cloth.GetSimPattern(SeamPattern0);
+			const FCollectionClothSimPatternConstFacade ClothPattern1 = Cloth.GetSimPattern(SeamPattern1);
+
+			const int32 ClothPattern0VerticesOffset = ClothPattern0.GetSimVertices2DOffset();
+			const int32 ClothPattern1VerticesOffset = ClothPattern1.GetSimVertices2DOffset();
+
+			TArray<FIntVector2> Stitches;
+			const uint32 StitchesCount = FMath::Min(SeamInfo.Seam0MeshIndices.Num(), SeamInfo.Seam1MeshIndices.Num());
+			Stitches.Reserve(StitchesCount);
+			for (uint32 StitchIndex = 0; StitchIndex < StitchesCount; ++StitchIndex)
+			{
+				Stitches.Emplace(
+					(int32)SeamInfo.Seam0MeshIndices[StitchIndex] + ClothPattern0VerticesOffset,
+					(int32)SeamInfo.Seam1MeshIndices[StitchIndex] + ClothPattern1VerticesOffset);
+			}
+
+			FCollectionClothSeamFacade SeamFacade = Cloth.AddGetSeam();
+			SeamFacade.Initialize(Stitches);
+		}
+	}
+
+	// Set the render mesh to duplicate the sim mesh
+	ChaosClothAsset->CopySimMeshToRenderMesh();
+
+	// Set a default skeleton and rebuild the asset
+	ChaosClothAsset->SetReferenceSkeleton(nullptr);  // This creates a default reference skeleton, redoes the bindings, and rebuilds the asset
+}
+
+UChaosClothAssetDatasmithClothComponentFactory::UChaosClothAssetDatasmithClothComponentFactory() = default;
+UChaosClothAssetDatasmithClothComponentFactory::~UChaosClothAssetDatasmithClothComponentFactory() = default;
+
+USceneComponent* UChaosClothAssetDatasmithClothComponentFactory::CreateClothComponent(UObject* Outer) const
+{
+	return CastChecked<USceneComponent>(NewObject<UChaosClothComponent>(Outer));
+}
+
+void UChaosClothAssetDatasmithClothComponentFactory::InitializeClothComponent(USceneComponent* ClothComponent, UObject* ClothAsset, USceneComponent* RootComponent) const
+{
+	UChaosClothComponent* const ChaosClothComponent = CastChecked<UChaosClothComponent>(ClothComponent);
+	ChaosClothComponent->SetClothAsset(Cast<UChaosClothAsset>(ClothAsset));
+	ChaosClothComponent->SetupAttachment(RootComponent);
+}
+
+namespace UE::Chaos::ClothAsset::Private
+{
+	/** A modular interface to provide and initialize cloth asset objects. */
+	class FDatasmithClothFactoryClassesProvider final: public IDatasmithClothFactoryClassesProvider
+	{
+	public:
+		FDatasmithClothFactoryClassesProvider() = default;
+		virtual ~FDatasmithClothFactoryClassesProvider() override = default;
+
+		virtual FName GetName() const override
+		{
+			static const FName Name = TEXT("ChaosClothAsset");
+			return Name;
+		}
+
+		virtual TSubclassOf<UDatasmithClothAssetFactory> GetClothAssetFactoryClass() const override
+		{
+			return TSubclassOf<UDatasmithClothAssetFactory>(UChaosClothAssetDatasmithClothAssetFactory::StaticClass());
+		}
+
+		virtual TSubclassOf<UDatasmithClothComponentFactory> GetClothComponentFactoryClass() const override
+		{
+			return TSubclassOf<UDatasmithClothComponentFactory>(UChaosClothAssetDatasmithClothComponentFactory::StaticClass());
+		}
+	};
+
+	static FDatasmithClothFactoryClassesProvider DatasmithClothFactoryClassesProvider;
+}
+
+void FChaosClothAssetDatasmithImportNode::RegisterModularFeature()
+{
+	using namespace UE::Chaos::ClothAsset::Private;
+	IModularFeatures::Get().RegisterModularFeature(IDatasmithClothFactoryClassesProvider::FeatureName, &DatasmithClothFactoryClassesProvider);
+}
+
+void FChaosClothAssetDatasmithImportNode::UnregisterModularFeature()
+{
+	using namespace UE::Chaos::ClothAsset::Private;
+	IModularFeatures::Get().UnregisterModularFeature(IDatasmithClothFactoryClassesProvider::FeatureName, &DatasmithClothFactoryClassesProvider);
+}
+
 FChaosClothAssetDatasmithImportNode::FChaosClothAssetDatasmithImportNode(const Dataflow::FNodeParameters& InParam, FGuid InGuid)
 	: FDataflowNode(InParam, InGuid)
 {
diff --git a/Engine/Plugins/Experimental/ChaosClothAssetEditor/Source/ChaosClothAssetDataflowNodes/Public/ChaosClothAsset/DatasmithImportNode.h b/Engine/Plugins/Experimental/ChaosClothAssetEditor/Source/ChaosClothAssetDataflowNodes/Public/ChaosClothAsset/DatasmithImportNode.h
--- a/Engine/Plugins/Experimental/ChaosClothAssetEditor/Source/ChaosClothAssetDataflowNodes/Public/ChaosClothAsset/DatasmithImportNode.h
+++ b/Engine/Plugins/Experimental/ChaosClothAssetEditor/Source/ChaosClothAssetDataflowNodes/Public/ChaosClothAsset/DatasmithImportNode.h
@@ -5,8 +5,39 @@
 #include "Dataflow/DataflowNode.h"
 #include "GeometryCollection/ManagedArrayCollection.h"
 #include "Misc/SecureHash.h"
+#include "DatasmithCloth.h"
 #include "DatasmithImportNode.generated.h"
 
+/** Cloth asset factory providing and initializing cloth assets on behalf of the Datasmith importer. */
+UCLASS()
+class UChaosClothAssetDatasmithClothAssetFactory final : public UDatasmithClothAssetFactory
+{
+	GENERATED_BODY()
+
+public:
+	UChaosClothAssetDatasmithClothAssetFactory();
+	virtual ~UChaosClothAssetDatasmithClothAssetFactory() override;
+
+	virtual UObject* CreateClothAsset(UObject* Outer, const FName& Name, EObjectFlags ObjectFlags) const override;
+	virtual UObject* DuplicateClothAsset(UObject* ClothAsset, UObject* Outer, const FName& Name) const override;
+	virtual void InitializeClothAsset(UObject* ClothAsset, const FDatasmithCloth& DatasmithCloth) const override;
+};
+
+
+/** Cloth component factory providing and initializing cloth components on behalf of the Datasmith importer. */
+UCLASS()
+class UChaosClothAssetDatasmithClothComponentFactory final : public UDatasmithClothComponentFactory
+{
+	GENERATED_BODY()
+
+public:
+	UChaosClothAssetDatasmithClothComponentFactory();
+	virtual ~UChaosClothAssetDatasmithClothComponentFactory() override;
+
+	virtual USceneComponent* CreateClothComponent(UObject* Outer) const override;
+	virtual void InitializeClothComponent(class USceneComponent* ClothComponent, UObject* ClothAsset, class USceneComponent* RootComponent) const override;
+};
+
 /** Import a file from a third party garment construction package compatible with the Datasmith scene format. */
 USTRUCT(meta = (DataflowCloth))
 struct FChaosClothAssetDatasmithImportNode : public FDataflowNode
@@ -22,6 +53,11 @@ public:
 	UPROPERTY(EditAnywhere, Category = "Datasmith Import")
 	FFilePath ImportFile;
 
+	/** Register the Datasmith cloth factory classes provider. */
+	static void RegisterModularFeature();
+	/** Unregister the Datasmith cloth factory classes provider. */
+	static void UnregisterModularFeature();
+
 	FChaosClothAssetDatasmithImportNode(const Dataflow::FNodeParameters& InParam, FGuid InGuid = FGuid::NewGuid());
 
 private:
diff --git a/Engine/Source/Runtime/Datasmith/DatasmithCore/Public/DatasmithCloth.h b/Engine/Source/Runtime/Datasmith/DatasmithCore/Public/DatasmithCloth.h
--- a/Engine/Source/Runtime/Datasmith/DatasmithCore/Public/DatasmithCloth.h
+++ b/Engine/Source/Runtime/Datasmith/DatasmithCore/Public/DatasmithCloth.h
@@ -4,12 +4,15 @@
 #include "Containers/Array.h"
 #include "Containers/UnrealString.h"
 #include "CoreMinimal.h"
+#include "Features/IModularFeature.h"
 #include "Math/Vector.h"
 #include "Math/Vector2D.h"
 #include "Misc/TVariant.h"
 #include "UObject/NameTypes.h"
+#include "DatasmithCloth.generated.h"
 
 class FDatasmithMesh;
+class UActorComponent;
 
 
 struct FParameterData
@@ -85,3 +88,57 @@ public:
 	friend FArchive& operator<<(FArchive& Ar, FDatasmithCloth& Cloth);
 };
 
+
+/** Modular cloth asset factory base class. */
+UCLASS(Abstract)
+class DATASMITHCORE_API UDatasmithClothAssetFactory : public UObject
+{
+	GENERATED_BODY()
+
+public:
+	UDatasmithClothAssetFactory() = default;
+	virtual ~UDatasmithClothAssetFactory() = default;
+
+	virtual UObject* CreateClothAsset(UObject* Outer, const FName& Name, EObjectFlags Flags) const
+	PURE_VIRTUAL(UDatasmithClothAssetFactory::CreateClothAsset, return nullptr;);
+
+	virtual UObject* DuplicateClothAsset(UObject* ClothAsset, UObject* Outer, const FName& Name) const
+	PURE_VIRTUAL(UDatasmithClothAssetFactory::DuplicateClothAsset, return nullptr;);
+
+	virtual void InitializeClothAsset(UObject* ClothAsset, const FDatasmithCloth& DatasmithCloth) const
+	PURE_VIRTUAL(UDatasmithClothAssetFactory::InitializeClothAsset, );
+};
+
+
+/** Modular cloth component factory base class. */
+UCLASS(Abstract)
+class DATASMITHCORE_API UDatasmithClothComponentFactory : public UObject
+{
+	GENERATED_BODY()
+
+public:
+	UDatasmithClothComponentFactory() = default;
+	virtual ~UDatasmithClothComponentFactory() = default;
+
+	virtual USceneComponent* CreateClothComponent(UObject* Outer) const
+	PURE_VIRTUAL(UDatasmithClothComponentFactory::CreateClothComponent, return nullptr;);
+
+	virtual void InitializeClothComponent(USceneComponent* ClothComponent, UObject* ClothAsset, USceneComponent* RootComponent) const
+	PURE_VIRTUAL(UDatasmithClothComponentFactory::InitializeClothComponent, );
+};
+
+
+/** A modular interface to provide factory classes to initialize cloth assets and components. */
+class IDatasmithClothFactoryClassesProvider : public IModularFeature
+{
+public:
+	inline static const FName FeatureName = TEXT("IDatasmithClothFactoryClassesProvider");
+
+	IDatasmithClothFactoryClassesProvider() = default;
+	virtual ~IDatasmithClothFactoryClassesProvider() = default;
+
+	virtual FName GetName() const = 0;
+
+	virtual TSubclassOf<UDatasmithClothAssetFactory> GetClothAssetFactoryClass() const = 0;
+	virtual TSubclassOf<UDatasmithClothComponentFactory> GetClothComponentFactoryClass() const = 0;
+};
