Remove impossible cast, UMaterial is not a UMaterialInstance
#rnx

diff --git a/Engine/Plugins/Enterprise/DatasmithImporter/Source/DatasmithImporter/Private/DatasmithConsumer.cpp b/Engine/Plugins/Enterprise/DatasmithImporter/Source/DatasmithImporter/Private/DatasmithConsumer.cpp
--- a/Engine/Plugins/Enterprise/DatasmithImporter/Source/DatasmithImporter/Private/DatasmithConsumer.cpp
+++ b/Engine/Plugins/Enterprise/DatasmithImporter/Source/DatasmithImporter/Private/DatasmithConsumer.cpp
@@ -1664,12 +1664,6 @@ namespace DatasmithConsumerUtils
 					TSharedRef< IDatasmithBaseMaterialElement > MaterialElement = FDatasmithSceneFactory::CreateMaterial( *AssetTag );
 					MaterialElement->SetLabel( *MaterialInstance->GetName() );
 
-					if (UMaterial* SourceMaterial = Cast< UMaterial >(MaterialInstance))
-					{
-						MaterialElement = StaticCastSharedRef< IDatasmithBaseMaterialElement >( FDatasmithSceneFactory::CreateUEPbrMaterial( *AssetTag ) );
-						MaterialElement->SetLabel( *MaterialInstance->GetName() );
-					}
-
 					if ( UMaterialInterface* MaterialParent = MaterialInstance->Parent )
 					{
 						FString MaterialInstancePath = MaterialInstance->GetOutermost()->GetName();
diff --git a/Engine/Plugins/Enterprise/DatasmithImporter/Source/DatasmithImporter/Private/DatasmithFileProducer.cpp b/Engine/Plugins/Enterprise/DatasmithImporter/Source/DatasmithImporter/Private/DatasmithFileProducer.cpp
--- a/Engine/Plugins/Enterprise/DatasmithImporter/Source/DatasmithImporter/Private/DatasmithFileProducer.cpp
+++ b/Engine/Plugins/Enterprise/DatasmithImporter/Source/DatasmithImporter/Private/DatasmithFileProducer.cpp
@@ -597,17 +597,6 @@ void UDatasmithFileProducer::PreventNameCollision()
 				}
 				else if( UMaterialInstance* MaterialInstance = Cast<UMaterialInstance>(Object) )
 				{
-					if (UMaterial* SourceMaterial = Cast< UMaterial >(MaterialInstance))
-					{
-						for (const FMaterialFunctionInfo& MaterialFunctionInfo : SourceMaterial->GetCachedExpressionData().FunctionInfos)
-						{
-							if (MaterialFunctionInfo.Function && MaterialFunctionInfo.Function->GetOutermost() == SourceMaterial->GetOutermost())
-							{
-								MaterialFunctions.Add( MaterialFunctionInfo.Function );
-							}
-						}
-					}
-
 					if ( UMaterialInterface* MaterialParent = MaterialInstance->Parent )
 					{
 						FString MaterialInstancePath = MaterialInstance->GetOutermost()->GetName();
