Finalize fix for  missing addition of certains types of actors to the import context in CL 27223391

#rnx

diff --git a/Engine/Plugins/Enterprise/DatasmithImporter/Source/DatasmithImporter/Private/DatasmithActorImporter.cpp b/Engine/Plugins/Enterprise/DatasmithImporter/Source/DatasmithImporter/Private/DatasmithActorImporter.cpp
--- a/Engine/Plugins/Enterprise/DatasmithImporter/Source/DatasmithImporter/Private/DatasmithActorImporter.cpp
+++ b/Engine/Plugins/Enterprise/DatasmithImporter/Source/DatasmithImporter/Private/DatasmithActorImporter.cpp
@@ -135,11 +135,6 @@ AActor* FDatasmithActorImporter::ImportActor( UClass* ActorClass, const TSharedR
 
 		SetupActorProperties( ImportedActor, ActorElement, ImportContext ) ;
 		SetupSceneComponent( ImportedActor->GetRootComponent(), ActorElement, ImportContext.Hierarchy.Num() > 0 ? ImportContext.Hierarchy.Top() : nullptr );
-
-		if ( DatasmithSceneActor )
-		{
-			DatasmithSceneActor->RelatedActors.FindOrAdd( ActorElement->GetName() ) = ImportedActor;
-		}
 	}
 
 	return ImportedActor;
diff --git a/Engine/Plugins/Enterprise/DatasmithImporter/Source/DatasmithImporter/Private/DatasmithImporter.cpp b/Engine/Plugins/Enterprise/DatasmithImporter/Source/DatasmithImporter/Private/DatasmithImporter.cpp
--- a/Engine/Plugins/Enterprise/DatasmithImporter/Source/DatasmithImporter/Private/DatasmithImporter.cpp
+++ b/Engine/Plugins/Enterprise/DatasmithImporter/Source/DatasmithImporter/Private/DatasmithImporter.cpp
@@ -1158,6 +1158,11 @@ AActor* FDatasmithImporter::ImportActor( FDatasmithImportContext& ImportContext,
 		ImportMetaDataForObject(ImportContext, ActorElement, ImportedActor);
 
 		ImportContext.AddImportedActor(ImportedActor);
+
+		if (ADatasmithSceneActor* DatasmithSceneActor = ImportContext.ActorsContext.ImportSceneActor)
+		{
+			DatasmithSceneActor->RelatedActors.FindOrAdd(ActorElement->GetName()) = ImportedActor;
+		}
 	}
 	else
 	{
