Fixed regression introduced by CLs 27457226 and 27223429

#rnx

diff --git a/Engine/Plugins/Enterprise/DatasmithImporter/Source/DatasmithImporter/Private/DatasmithActorImporter.cpp b/Engine/Plugins/Enterprise/DatasmithImporter/Source/DatasmithImporter/Private/DatasmithActorImporter.cpp
--- a/Engine/Plugins/Enterprise/DatasmithImporter/Source/DatasmithImporter/Private/DatasmithActorImporter.cpp
+++ b/Engine/Plugins/Enterprise/DatasmithImporter/Source/DatasmithImporter/Private/DatasmithActorImporter.cpp
@@ -135,6 +135,13 @@ AActor* FDatasmithActorImporter::ImportActor( UClass* ActorClass, const TSharedR
 
 		SetupActorProperties( ImportedActor, ActorElement, ImportContext ) ;
 		SetupSceneComponent( ImportedActor->GetRootComponent(), ActorElement, ImportContext.Hierarchy.Num() > 0 ? ImportContext.Hierarchy.Top() : nullptr );
+
+		ImportContext.AddImportedActor( ImportedActor );
+
+		if ( DatasmithSceneActor )
+		{
+			DatasmithSceneActor->RelatedActors.FindOrAdd( ActorElement->GetName() ) = ImportedActor;
+		}
 	}
 
 	return ImportedActor;
