Fixed crash when actor's label is longer than 1024 characters

#rnx

diff --git a/Engine/Plugins/Enterprise/DatasmithImporter/Source/DatasmithImporter/Private/DatasmithActorImporter.cpp b/Engine/Plugins/Enterprise/DatasmithImporter/Source/DatasmithImporter/Private/DatasmithActorImporter.cpp
--- a/Engine/Plugins/Enterprise/DatasmithImporter/Source/DatasmithImporter/Private/DatasmithActorImporter.cpp
+++ b/Engine/Plugins/Enterprise/DatasmithImporter/Source/DatasmithImporter/Private/DatasmithActorImporter.cpp
@@ -45,9 +45,49 @@
 
 #define LOCTEXT_NAMESPACE "DatasmithActorImporter"
 
+namespace UE::DatasmithActorImporter
+{
+	//
+	// For base actors, Datasmith is using actor's label to name scene components.
+	// If the label is too long, the scene component's name or the one of it is an outer to
+	// may exceed the NAME_SIZE limit and cause a crash.
+	// The deepest path would end up being:
+	// /Game/[CONTENT_FOLDER_PATH]/[LevelName].[LevelName]:PersistentLevel.[SceneComponentName].DatasmithAssetUserData_0.DatasmithActorTemplate_0
+	// To avoid such crash, the scene component's name will be truncated to allow the append of ".DatasmithAssetUserData_0.DatasmithActorTemplate_0"
+	// In case a potential crash is identified,the scene component's name will be composed of at max
+	// the first 50 characters of the desired name followed by a text indicating it has been truncated and terminated with
+	// at max the 50 last characters of the desired name
+	FString ComputeSceneComponentName( const UObject& Outer, const FString& DesiredName, FDatasmithImportContext& ImportContext )
+	{
+		FString PathFullName = Outer.GetFullName() / DesiredName;
+		constexpr int32 ExpectedSuffixMaxLength = 51; // Length of ".DatasmithAssetUserData_0.DatasmithActorTemplate_0"
+		int32 PathLength = PathFullName.Len() + ExpectedSuffixMaxLength;
+		constexpr int32 PreventiveMaximumLength = 900;
+		if ( PathLength < PreventiveMaximumLength)
+		{
+			return DesiredName;
+		}
+
+		// Known final path length = Outer's path length + 1 + length of ".DatasmithAssetUserData_0.DatasmithActorTemplate_0" + length of included warning text
+		constexpr int32 WarningtextLength = 23; // Length of warning text "---NAME---TRUNCATED---" inserted in new name
+		const int32 KnownPathLength = Outer.GetFullName().Len() + ExpectedSuffixMaxLength + WarningtextLength + 1;
+		const int32 RemainingLength = FMath::Min(PreventiveMaximumLength - KnownPathLength, 100);
+
+		FString ComputedName(RemainingLength / 2, *DesiredName);
+		ComputedName += TEXT("---NAME---TRUNCATED---");
+		ComputedName += DesiredName.Right(RemainingLength / 2);
+
+		ImportContext.LogWarning(FText::Format(LOCTEXT("ComputeSceneComponentName", "Detected potential long name, {0}, for scene component. Truncating..."), FText::FromString(DesiredName)));
+
+		return ComputedName;
+	}
+}
+
 AActor* FDatasmithActorImporter::ImportActor( UClass* ActorClass, const TSharedRef< IDatasmithActorElement >& ActorElement, FDatasmithImportContext& ImportContext, EDatasmithImportActorPolicy ImportActorPolicy,
 	TFunction< void( AActor* ) > PostSpawnFunc )
 {
+	using namespace UE::DatasmithActorImporter;
+
 	if ( ImportActorPolicy == EDatasmithImportActorPolicy::Ignore )
 	{
 		return nullptr;
@@ -120,7 +160,9 @@ AActor* FDatasmithActorImporter::ImportActor( UClass* ActorClass, const TSharedR
 
 		if ( !RootComponent )
 		{
-			RootComponent = NewObject< USceneComponent >( ImportedActor, USceneComponent::StaticClass(), ActorElement->GetLabel(), RF_Transactional );
+			FString ComponentName = ComputeSceneComponentName( *ImportedActor, ActorElement->GetLabel(), ImportContext );
+
+			RootComponent = NewObject< USceneComponent >( ImportedActor, USceneComponent::StaticClass(), *ComponentName, RF_Transactional );
 
 			ImportedActor->AddInstanceComponent( RootComponent );
 			ImportedActor->SetRootComponent( RootComponent );
@@ -149,6 +191,8 @@ AActor* FDatasmithActorImporter::ImportActor( UClass* ActorClass, const TSharedR
 
 USceneComponent* FDatasmithActorImporter::ImportSceneComponent( UClass* ComponentClass, const TSharedRef< IDatasmithActorElement >& ActorElement, FDatasmithImportContext& ImportContext, UObject* Outer, FDatasmithActorUniqueLabelProvider& UniqueNameProvider )
 {
+	using namespace UE::DatasmithActorImporter;
+
 	if ( !ComponentClass->IsChildOf( USceneComponent::StaticClass() ) )
 	{
 		ensure( false );
@@ -162,6 +206,10 @@ USceneComponent* FDatasmithActorImporter::ImportSceneComponent( UClass* Componen
 	}
 
 	AActor* Actor = Cast< AActor >( Outer );
+	if (!ensure(Actor))
+	{
+		return nullptr;
+	}
 
 	// This is possibly the SceneComponent we are looking for as the existing component
 	USceneComponent* SceneComponent = static_cast< USceneComponent* >( FindObjectWithOuter( Outer, ComponentClass, ActorElement->GetLabel() ) );
@@ -198,7 +246,7 @@ USceneComponent* FDatasmithActorImporter::ImportSceneComponent( UClass* Componen
 
 	if ( !ValidSceneComponent )
 	{
-		FName ComponentName( ActorElement->GetLabel() );
+		FName ComponentName( *ComputeSceneComponentName( *Outer, ActorElement->GetLabel(), ImportContext ) );
 		if ( SceneComponent || FindObjectWithOuter( Outer, UObject::StaticClass(), ComponentName ) )
 		{
 			// There is already a object with this name inside the outer. Generate unique name.
