Fix the deprecated warning caused by GetHttpReceiveTimeout, use GetHttpActivityTimeout instead.

diff --git a/Engine/Plugins/Enterprise/DatasmithImporter/Source/DatasmithImporter/Private/UI/DatasmithImportOptionsWindow.cpp b/Engine/Plugins/Enterprise/DatasmithImporter/Source/DatasmithImporter/Private/UI/DatasmithImportOptionsWindow.cpp
--- a/Engine/Plugins/Enterprise/DatasmithImporter/Source/DatasmithImporter/Private/UI/DatasmithImportOptionsWindow.cpp
+++ b/Engine/Plugins/Enterprise/DatasmithImporter/Source/DatasmithImporter/Private/UI/DatasmithImportOptionsWindow.cpp
@@ -306,7 +306,7 @@ bool FDatasmithImportOptionHelpURLTestHelper::SendHttpRequest()
 	});
 
 	TOptional<float> RequestTimeout = HttpRequest->GetTimeout();
-	MaxTimeoutDelay = RequestTimeout.IsSet() ? RequestTimeout.GetValue() : HttpModule.GetHttpReceiveTimeout();
+	MaxTimeoutDelay = RequestTimeout.IsSet() ? RequestTimeout.GetValue() : HttpModule.GetHttpActivityTimeout();
 
 	return HttpRequest->ProcessRequest();
 }
