Change FTickableGameObjects that use the soon-to-be-deprecated IsAllowedToTick to instead use IsTickable if it is actually needed

diff --git a/Engine/Plugins/EnhancedInput/Source/InputEditor/Private/EnhancedInputEditorSubsystem.cpp b/Engine/Plugins/EnhancedInput/Source/InputEditor/Private/EnhancedInputEditorSubsystem.cpp
--- a/Engine/Plugins/EnhancedInput/Source/InputEditor/Private/EnhancedInputEditorSubsystem.cpp
+++ b/Engine/Plugins/EnhancedInput/Source/InputEditor/Private/EnhancedInputEditorSubsystem.cpp
@@ -89,7 +89,7 @@ ETickableTickType UEnhancedInputEditorSubsystem::GetTickableTickType() const
 	return IsTemplate() ? ETickableTickType::Never : ETickableTickType::Conditional;
 }
 
-bool UEnhancedInputEditorSubsystem::IsAllowedToTick() const
+bool UEnhancedInputEditorSubsystem::IsTickable() const
 {
 	// Only tick if we have created a valid player input (i.e. we have been initalized)
 	return PlayerInput != nullptr && bIsCurrentlyConsumingInput;
diff --git a/Engine/Plugins/EnhancedInput/Source/InputEditor/Public/EnhancedInputEditorSubsystem.h b/Engine/Plugins/EnhancedInput/Source/InputEditor/Public/EnhancedInputEditorSubsystem.h
--- a/Engine/Plugins/EnhancedInput/Source/InputEditor/Public/EnhancedInputEditorSubsystem.h
+++ b/Engine/Plugins/EnhancedInput/Source/InputEditor/Public/EnhancedInputEditorSubsystem.h
@@ -40,7 +40,7 @@ public:
 	virtual UWorld* GetTickableGameObjectWorld() const override;
 	virtual bool IsTickableInEditor() const { return true; }
 	virtual ETickableTickType GetTickableTickType() const override;
-	virtual bool IsAllowedToTick() const override;
+	virtual bool IsTickable() const override;
 	virtual void Tick(float DeltaTime) override;
 	TStatId GetStatId() const override { RETURN_QUICK_DECLARE_CYCLE_STAT(UEnhancedInputEditorSubsystem, STATGROUP_Tickables); }
 	//~ End FTickableGameObject interface
diff --git a/Engine/Plugins/Enterprise/DatasmithImporter/Source/DirectLinkExtension/Private/AutoReimportManager.h b/Engine/Plugins/Enterprise/DatasmithImporter/Source/DirectLinkExtension/Private/AutoReimportManager.h
--- a/Engine/Plugins/Enterprise/DatasmithImporter/Source/DirectLinkExtension/Private/AutoReimportManager.h
+++ b/Engine/Plugins/Enterprise/DatasmithImporter/Source/DirectLinkExtension/Private/AutoReimportManager.h
@@ -23,7 +23,7 @@ namespace UE::DatasmithImporter
 	public:
 		/* Begin FTickableGameObject Interface */
 		virtual void Tick(float DeltaTime) override;
-		virtual bool IsAllowedToTick() const override { return PendingAutoReimportObjects.Num() > 0 || !PendingReimportQueue.IsEmpty() || !PendingInvalidateQueue.IsEmpty(); }
+		virtual bool IsTickable() const override { return PendingAutoReimportObjects.Num() > 0 || !PendingReimportQueue.IsEmpty() || !PendingInvalidateQueue.IsEmpty(); }
 		virtual TStatId GetStatId() const override;
 		virtual bool IsTickableWhenPaused() const override { return true; }
 		virtual bool IsTickableInEditor() const override { return true; }
diff --git a/Engine/Plugins/Experimental/RenderGrid/Source/RenderGrid/Public/RenderGrid/RenderGridQueue.h b/Engine/Plugins/Experimental/RenderGrid/Source/RenderGrid/Public/RenderGrid/RenderGridQueue.h
--- a/Engine/Plugins/Experimental/RenderGrid/Source/RenderGrid/Public/RenderGrid/RenderGridQueue.h
+++ b/Engine/Plugins/Experimental/RenderGrid/Source/RenderGrid/Public/RenderGrid/RenderGridQueue.h
@@ -232,8 +232,6 @@ public:
 	virtual ETickableTickType GetTickableTickType() const override { return ETickableTickType::Always; }
 	virtual bool IsTickableWhenPaused() const override { return true; }
 	virtual bool IsTickableInEditor() const override { return true; }
-	virtual bool IsTickable() const override { return true; }
-	virtual bool IsAllowedToTick() const override { return true; }
 	virtual TStatId GetStatId() const override
 	{
 		RETURN_QUICK_DECLARE_CYCLE_STAT(FRenderGridQueue, STATGROUP_Tickables);
diff --git a/Engine/Plugins/Experimental/RenderGrid/Source/RenderGrid/Public/Utils/RenderGridGenericExecutionQueue.h b/Engine/Plugins/Experimental/RenderGrid/Source/RenderGrid/Public/Utils/RenderGridGenericExecutionQueue.h
--- a/Engine/Plugins/Experimental/RenderGrid/Source/RenderGrid/Public/Utils/RenderGridGenericExecutionQueue.h
+++ b/Engine/Plugins/Experimental/RenderGrid/Source/RenderGrid/Public/Utils/RenderGridGenericExecutionQueue.h
@@ -100,8 +100,6 @@ namespace UE::RenderGrid::Private
 		virtual ETickableTickType GetTickableTickType() const override { return ETickableTickType::Always; }
 		virtual bool IsTickableWhenPaused() const override { return true; }
 		virtual bool IsTickableInEditor() const override { return true; }
-		virtual bool IsTickable() const override { return true; }
-		virtual bool IsAllowedToTick() const override { return true; }
 		virtual TStatId GetStatId() const override
 		{
 			RETURN_QUICK_DECLARE_CYCLE_STAT(FRenderGridGenericExecutionQueue, STATGROUP_Tickables);
diff --git a/Engine/Plugins/Experimental/Web/WebAPI/Source/WebAPI/Private/WebAPISubsystem.cpp b/Engine/Plugins/Experimental/Web/WebAPI/Source/WebAPI/Private/WebAPISubsystem.cpp
--- a/Engine/Plugins/Experimental/Web/WebAPI/Source/WebAPI/Private/WebAPISubsystem.cpp
+++ b/Engine/Plugins/Experimental/Web/WebAPI/Source/WebAPI/Private/WebAPISubsystem.cpp
@@ -140,12 +140,11 @@ void UWebAPISubsystem::Deinitialize()
 
 ETickableTickType UWebAPISubsystem::GetTickableTickType() const
 {
-	return FTickableGameObject::GetTickableTickType();
-}
-
-bool UWebAPISubsystem::IsAllowedToTick() const
-{
-	return FTickableGameObject::IsAllowedToTick();
+	if (IsTemplate())
+	{
+		return ETickableTickType::Never;
+	}
+	return ETickableTickType::Always;
 }
 
 void UWebAPISubsystem::Tick(float DeltaTime)
diff --git a/Engine/Plugins/Experimental/Web/WebAPI/Source/WebAPI/Public/WebAPISubsystem.h b/Engine/Plugins/Experimental/Web/WebAPI/Source/WebAPI/Public/WebAPISubsystem.h
--- a/Engine/Plugins/Experimental/Web/WebAPI/Source/WebAPI/Public/WebAPISubsystem.h
+++ b/Engine/Plugins/Experimental/Web/WebAPI/Source/WebAPI/Public/WebAPISubsystem.h
@@ -81,7 +81,6 @@ public:
 	// FTickableGameObject implementation Begin
 	virtual UWorld* GetTickableGameObjectWorld() const override { return GetWorld(); }
 	virtual ETickableTickType GetTickableTickType() const override;
-	virtual bool IsAllowedToTick() const override final;
 	virtual void Tick(float DeltaTime) override;
 	virtual TStatId GetStatId() const override;
 	// FTickableGameObject implementation End
diff --git a/Engine/Plugins/Media/ElectraPlayer/Source/ElectraPlayerRuntime/Private/Runtime/Player/SimpleElectraAudioPlayer.cpp b/Engine/Plugins/Media/ElectraPlayer/Source/ElectraPlayerRuntime/Private/Runtime/Player/SimpleElectraAudioPlayer.cpp
--- a/Engine/Plugins/Media/ElectraPlayer/Source/ElectraPlayerRuntime/Private/Runtime/Player/SimpleElectraAudioPlayer.cpp
+++ b/Engine/Plugins/Media/ElectraPlayer/Source/ElectraPlayerRuntime/Private/Runtime/Player/SimpleElectraAudioPlayer.cpp
@@ -191,11 +191,7 @@ class FSimpleElectraAudioPlayer : public ISimpleElectraAudioPlayer, public IAdap
 			FSimpleElectraAudioPlayer::TickAllInstances(DeltaTime);
 		}
 		ETickableTickType GetTickableTickType() const override
-		{ return ETickableTickType::Conditional; }
-		bool IsTickable() const override
-		{ return true; }
-		bool IsAllowedToTick() const override
-		{ return true; }
+		{ return ETickableTickType::Always; }
 		TStatId GetStatId() const override
 		{ RETURN_QUICK_DECLARE_CYCLE_STAT(FSimpleElectraAudioPlayer, STATGROUP_Tickables); }
 		bool IsTickableWhenPaused() const override
diff --git a/Engine/Plugins/Media/PixelStreaming/Source/PixelStreamingServers/Private/ServerBase.cpp b/Engine/Plugins/Media/PixelStreaming/Source/PixelStreamingServers/Private/ServerBase.cpp
--- a/Engine/Plugins/Media/PixelStreaming/Source/PixelStreamingServers/Private/ServerBase.cpp
+++ b/Engine/Plugins/Media/PixelStreaming/Source/PixelStreamingServers/Private/ServerBase.cpp
@@ -31,7 +31,7 @@ namespace UE::PixelStreamingServers
 		return true;
 	}
 
-	bool FServerBase::IsAllowedToTick() const
+	bool FServerBase::IsTickable() const
 	{
 		return bAllowedToTick;
 	}
diff --git a/Engine/Plugins/Media/PixelStreaming/Source/PixelStreamingServers/Private/ServerBase.h b/Engine/Plugins/Media/PixelStreaming/Source/PixelStreamingServers/Private/ServerBase.h
--- a/Engine/Plugins/Media/PixelStreaming/Source/PixelStreamingServers/Private/ServerBase.h
+++ b/Engine/Plugins/Media/PixelStreaming/Source/PixelStreamingServers/Private/ServerBase.h
@@ -28,7 +28,7 @@ namespace UE::PixelStreamingServers
 		virtual bool IsTickableWhenPaused() const override;
 		virtual bool IsTickableInEditor() const override;
 		virtual void Tick(float DeltaTime) override;
-		virtual bool IsAllowedToTick() const override;
+		virtual bool IsTickable() const override;
 		TStatId GetStatId() const { RETURN_QUICK_DECLARE_CYCLE_STAT(PixelStreamingServers, STATGROUP_Tickables); }
 		/* End FTickableGameObject */
 	
diff --git a/Engine/Plugins/Media/PixelStreaming/Source/PixelStreamingServers/Private/WebSocketServerWrapper.h b/Engine/Plugins/Media/PixelStreaming/Source/PixelStreamingServers/Private/WebSocketServerWrapper.h
--- a/Engine/Plugins/Media/PixelStreaming/Source/PixelStreamingServers/Private/WebSocketServerWrapper.h
+++ b/Engine/Plugins/Media/PixelStreaming/Source/PixelStreamingServers/Private/WebSocketServerWrapper.h
@@ -66,8 +66,8 @@ namespace UE::PixelStreamingServers
 		/* Begin FTickableGameObject */
 		virtual bool IsTickableWhenPaused() const { return true; }
 		virtual bool IsTickableInEditor() const { return true; }
+		virtual ETickableTickType GetTickableTickType() const override { return ETickableTickType::Always; }
 		virtual void Tick(float DeltaTime) override;
-		virtual bool IsAllowedToTick() const { return true; }
 		TStatId GetStatId() const { RETURN_QUICK_DECLARE_CYCLE_STAT(FWebSocketServerWrapper, STATGROUP_Tickables); }
 		/* End FTickableGameObject */
 
diff --git a/Engine/Source/Editor/DataLayerEditor/Private/DataLayer/DataLayerEditorSubsystem.cpp b/Engine/Source/Editor/DataLayerEditor/Private/DataLayer/DataLayerEditorSubsystem.cpp
--- a/Engine/Source/Editor/DataLayerEditor/Private/DataLayer/DataLayerEditorSubsystem.cpp
+++ b/Engine/Source/Editor/DataLayerEditor/Private/DataLayer/DataLayerEditorSubsystem.cpp
@@ -406,7 +406,7 @@ ETickableTickType UDataLayerEditorSubsystem::GetTickableTickType() const
 	return IsTemplate() ? ETickableTickType::Never : ETickableTickType::Conditional;
 }
 
-bool UDataLayerEditorSubsystem::IsAllowedToTick() const
+bool UDataLayerEditorSubsystem::IsTickable() const
 {
 	return GetWorld() && (bAsyncBroadcastDataLayerChanged || bAsyncUpdateAllActorsVisibility || bAsyncInvalidateViewports || LastWarningNotification.IsSet());
 }
diff --git a/Engine/Source/Editor/DataLayerEditor/Public/DataLayer/DataLayerEditorSubsystem.h b/Engine/Source/Editor/DataLayerEditor/Public/DataLayer/DataLayerEditorSubsystem.h
--- a/Engine/Source/Editor/DataLayerEditor/Public/DataLayer/DataLayerEditorSubsystem.h
+++ b/Engine/Source/Editor/DataLayerEditor/Public/DataLayer/DataLayerEditorSubsystem.h
@@ -99,7 +99,7 @@ public:
 	virtual UWorld* GetTickableGameObjectWorld() const override;
 	virtual bool IsTickableInEditor() const { return true; }
 	virtual ETickableTickType GetTickableTickType() const override;
-	virtual bool IsAllowedToTick() const override;
+	virtual bool IsTickable() const override;
 	virtual void Tick(float DeltaTime) override;
 	TStatId GetStatId() const override { RETURN_QUICK_DECLARE_CYCLE_STAT(UDataLayerEditorSubsystem, STATGROUP_Tickables); }
 	//~ End FTickableGameObject interface
diff --git a/Engine/Source/Runtime/Engine/Classes/GameFramework/InputDeviceSubsystem.h b/Engine/Source/Runtime/Engine/Classes/GameFramework/InputDeviceSubsystem.h
--- a/Engine/Source/Runtime/Engine/Classes/GameFramework/InputDeviceSubsystem.h
+++ b/Engine/Source/Runtime/Engine/Classes/GameFramework/InputDeviceSubsystem.h
@@ -141,7 +141,7 @@ public:
 	//~ Begin FTickableGameObject interface	
 	ENGINE_API virtual UWorld* GetTickableGameObjectWorld() const override;
 	ENGINE_API virtual ETickableTickType GetTickableTickType() const override;
-	ENGINE_API virtual bool IsAllowedToTick() const override;
+	ENGINE_API virtual bool IsTickable() const override;
 	ENGINE_API virtual bool IsTickableInEditor() const override;
 	ENGINE_API virtual TStatId GetStatId() const override;
 	ENGINE_API virtual void Tick(float InDeltaTime) override;
diff --git a/Engine/Source/Runtime/Engine/Private/GameFramework/InputDeviceSubsystem.cpp b/Engine/Source/Runtime/Engine/Private/GameFramework/InputDeviceSubsystem.cpp
--- a/Engine/Source/Runtime/Engine/Private/GameFramework/InputDeviceSubsystem.cpp
+++ b/Engine/Source/Runtime/Engine/Private/GameFramework/InputDeviceSubsystem.cpp
@@ -271,7 +271,7 @@ ETickableTickType UInputDeviceSubsystem::GetTickableTickType() const
 	return (IsTemplate() ? ETickableTickType::Never : ETickableTickType::Conditional);
 }
 
-bool UInputDeviceSubsystem::IsAllowedToTick() const
+bool UInputDeviceSubsystem::IsTickable() const
 {
 	// Only tick when there are active device properties or ones we want to remove
 	const bool bWantsTick = !ActiveProperties.IsEmpty() || !PropertiesPendingRemoval.IsEmpty();
