Fixed reimport of static meshes imported from a wire file
It was a noop in 5.4 but a crash in 5.5.

#rnx

diff --git a/Engine/Plugins/Enterprise/DatasmithImporter/Source/DatasmithImporter/Private/DatasmithImportContext.cpp b/Engine/Plugins/Enterprise/DatasmithImporter/Source/DatasmithImporter/Private/DatasmithImportContext.cpp
--- a/Engine/Plugins/Enterprise/DatasmithImporter/Source/DatasmithImporter/Private/DatasmithImportContext.cpp
+++ b/Engine/Plugins/Enterprise/DatasmithImporter/Source/DatasmithImporter/Private/DatasmithImportContext.cpp
@@ -20,6 +20,7 @@
 #include "IUriManager.h"
 #include "Misc/FeedbackContext.h"
 #include "Misc/Paths.h"
+#include "Misc/SecureHash.h"
 #include "PackageTools.h"
 #include "SourceUri.h"
 
@@ -647,6 +648,14 @@ bool FDatasmithAssetsImportContext::Init()
 
 void FDatasmithAssetsImportContext::ReInit(const FString& NewRootFolder)
 {
+	// Prevent re-initialize twice with the same root folder
+	// This happens when reimporting specific assets or scene
+	// This used to work but now crashes when creating the transient packages
+	if (RootFolderPath == NewRootFolder && StaticMeshesFinalPackage.IsValid())
+	{
+		return;
+	}
+
 	RootFolderPath = UPackageTools::SanitizePackageName(NewRootFolder);
 
 	StaticMeshesFinalPackage.Reset( CreatePackage( *FPaths::Combine( RootFolderPath, TEXT("Geometries") ) ) );
@@ -656,7 +665,8 @@ void FDatasmithAssetsImportContext::ReInit(const FString& NewRootFolder)
 	LevelSequencesFinalPackage.Reset( CreatePackage( *FPaths::Combine( RootFolderPath, TEXT("Animations") ) ) );
 	LevelVariantSetsFinalPackage.Reset( CreatePackage( *FPaths::Combine( RootFolderPath, TEXT("Variants") ) ) );
 
-	TransientFolderPath = FPaths::Combine( RootFolderPath, TEXT("Temp") );
+	// Use the engine's transient package path as initial root to create the Datasmith transient packages
+	TransientFolderPath = FPaths::Combine(GetTransientPackage()->GetPathName(), FMD5::HashAnsiString(*RootFolderPath));
 
 	StaticMeshesImportPackage.Reset( NewObject< UPackage >( nullptr, *FPaths::Combine( TransientFolderPath, TEXT("Geometries") ), RF_Transient ) );
 	StaticMeshesImportPackage->FullyLoad();
