Added Mobility attribute to IDatasmithActorElement

#rnx

diff --git a/Engine/Plugins/Enterprise/DatasmithImporter/Source/DatasmithImporter/Private/DatasmithActorImporter.cpp b/Engine/Plugins/Enterprise/DatasmithImporter/Source/DatasmithImporter/Private/DatasmithActorImporter.cpp
--- a/Engine/Plugins/Enterprise/DatasmithImporter/Source/DatasmithImporter/Private/DatasmithActorImporter.cpp
+++ b/Engine/Plugins/Enterprise/DatasmithImporter/Source/DatasmithImporter/Private/DatasmithActorImporter.cpp
@@ -749,6 +749,10 @@ void FDatasmithActorImporter::SetupActorProperties(AActor* ImportedActor, const
 
 void FDatasmithActorImporter::SetupSceneComponent( USceneComponent* SceneComponent, const TSharedRef< IDatasmithActorElement >& ActorElement, USceneComponent* Parent )
 {
+	static_assert((uint8)EDatasmithActorMobilityType::Static == (uint8)EComponentMobility::Type::Static, "ENUM_VALUE_HAS_CHANGED");
+	static_assert((uint8)EDatasmithActorMobilityType::Stationary == (uint8)EComponentMobility::Type::Stationary, "ENUM_VALUE_HAS_CHANGED");
+	static_assert((uint8)EDatasmithActorMobilityType::Movable == (uint8)EComponentMobility::Type::Movable, "ENUM_VALUE_HAS_CHANGED");
+
 	if ( !SceneComponent )
 	{
 		return;
@@ -758,7 +762,7 @@ void FDatasmithActorImporter::SetupSceneComponent( USceneComponent* SceneCompone
 	UDatasmithSceneComponentTemplate* SceneComponentTemplate = NewObject< UDatasmithSceneComponentTemplate >(Outer);
 
 	SceneComponentTemplate->RelativeTransform = ActorElement->GetRelativeTransform();
-	SceneComponentTemplate->Mobility = ActorElement->IsA(EDatasmithElementType::Camera) ? EComponentMobility::Movable : EComponentMobility::Static;
+	SceneComponentTemplate->Mobility = (EComponentMobility::Type)ActorElement->GetMobility();
 	SceneComponentTemplate->bVisible = ActorElement->GetVisibility();
 	SceneComponentTemplate->bCastShadow = ActorElement->GetCastShadow();
 	SceneComponentTemplate->AttachParent = Parent;
diff --git a/Engine/Source/Runtime/Datasmith/DatasmithCore/Private/DatasmithSceneElementsImpl.cpp b/Engine/Source/Runtime/Datasmith/DatasmithCore/Private/DatasmithSceneElementsImpl.cpp
--- a/Engine/Source/Runtime/Datasmith/DatasmithCore/Private/DatasmithSceneElementsImpl.cpp
+++ b/Engine/Source/Runtime/Datasmith/DatasmithCore/Private/DatasmithSceneElementsImpl.cpp
@@ -267,6 +267,8 @@ FDatasmithCameraActorElementImpl::FDatasmithCameraActorElementImpl(const TCHAR*
 	Store.RegisterParameter(FocalLength,         "FocalLength"        );
 	Store.RegisterParameter(ActorName,           "ActorName"          );
 	Store.RegisterParameter(bLookAtAllowRoll,    "bLookAtAllowRoll"   );
+
+	SetMobility(EDatasmithActorMobilityType::Movable);
 }
 
 float FDatasmithCameraActorElementImpl::GetSensorWidth() const
diff --git a/Engine/Source/Runtime/Datasmith/DatasmithCore/Private/DatasmithSceneElementsImpl.h b/Engine/Source/Runtime/Datasmith/DatasmithCore/Private/DatasmithSceneElementsImpl.h
--- a/Engine/Source/Runtime/Datasmith/DatasmithCore/Private/DatasmithSceneElementsImpl.h
+++ b/Engine/Source/Runtime/Datasmith/DatasmithCore/Private/DatasmithSceneElementsImpl.h
@@ -200,6 +200,9 @@ public:
 	virtual void SetCastShadow(bool bInCastShadow) override { bCastShadow = bInCastShadow; }
 	virtual bool GetCastShadow() const override { return bCastShadow; }
 
+	virtual void SetMobility(EDatasmithActorMobilityType InMobility) override { Mobility = InMobility; }
+	virtual EDatasmithActorMobilityType GetMobility() const override { return Mobility; }
+
 protected:
 	/** Converts all children's transforms to relative */
 	void ConvertChildsToRelative();
@@ -227,6 +230,7 @@ private:
 
 	TReflected<EActorFlags, uint8> Flags;
 	TReflected<bool> bCastShadow;
+	TReflected<EDatasmithActorMobilityType, uint8> Mobility;
 };
 
 template< typename InterfaceType >
@@ -240,6 +244,7 @@ inline FDatasmithActorElementImpl<T>::FDatasmithActorElementImpl(const TCHAR* In
 	, Rotation(FQuat::Identity)
 	, Flags(EActorFlags::IsVisible)
 	, bCastShadow(true)
+	, Mobility(EDatasmithActorMobilityType::Static)
 {
 	this->RegisterReferenceProxy(Children, "Children");
 	this->RegisterReferenceProxy(Parent,   "Parent"  );
@@ -251,6 +256,7 @@ inline FDatasmithActorElementImpl<T>::FDatasmithActorElementImpl(const TCHAR* In
 	Store.RegisterParameter(Tags,         "Tags"         ); // reflect as low prio for directlink
 	Store.RegisterParameter(Flags,        "Flags"        );
 	Store.RegisterParameter(bCastShadow,  "CastShadow"   );
+	Store.RegisterParameter(Mobility,     "Mobility"     );
 }
 
 template< typename T >
diff --git a/Engine/Source/Runtime/Datasmith/DatasmithCore/Private/DatasmithSceneXmlReader.cpp b/Engine/Source/Runtime/Datasmith/DatasmithCore/Private/DatasmithSceneXmlReader.cpp
--- a/Engine/Source/Runtime/Datasmith/DatasmithCore/Private/DatasmithSceneXmlReader.cpp
+++ b/Engine/Source/Runtime/Datasmith/DatasmithCore/Private/DatasmithSceneXmlReader.cpp
@@ -758,6 +758,18 @@ PRAGMA_ENABLE_DEPRECATION_WARNINGS
 	FString CastShadowAtribute = InNode->GetAttribute(TEXT("castshadow"));
 	InOutElement->SetCastShadow(CastShadowAtribute.IsEmpty() ? true : ValueFromString<bool>(CastShadowAtribute));
 
+	FString MobilityString = InNode->GetAttribute(TEXT("mobility"));
+	TArrayView< const TCHAR* > MobilityEnumStrings(DatasmithActorMobilityTypeStrings);
+	int32 MobilityIndexOfEnumValue = MobilityEnumStrings.IndexOfByPredicate([&MobilityString](const TCHAR* Value)
+		{
+			return MobilityString == Value;
+		});
+
+	if (MobilityIndexOfEnumValue != INDEX_NONE)
+	{
+		InOutElement->SetMobility((EDatasmithActorMobilityType)MobilityIndexOfEnumValue);
+	}
+
 	for (FXmlNode* ChildNode : InNode->GetChildrenNodes())
 	{
 		if (ChildNode->GetTag() == TEXT("transform"))
diff --git a/Engine/Source/Runtime/Datasmith/DatasmithCore/Private/DatasmithSceneXmlWriter.cpp b/Engine/Source/Runtime/Datasmith/DatasmithCore/Private/DatasmithSceneXmlWriter.cpp
--- a/Engine/Source/Runtime/Datasmith/DatasmithCore/Private/DatasmithSceneXmlWriter.cpp
+++ b/Engine/Source/Runtime/Datasmith/DatasmithCore/Private/DatasmithSceneXmlWriter.cpp
@@ -146,6 +146,11 @@ FString FDatasmithSceneXmlWriterImpl::GetLabelAndLayer(const TSharedPtr<IDatasmi
 		LabelAndLayer += TEXT(" component=\"true\"");
 	}
 
+	ensure((int32)ActorElement->GetMobility() < UE_ARRAY_COUNT(DatasmithActorMobilityTypeStrings));
+	FString MobilityString = DatasmithActorMobilityTypeStrings[(int32)ActorElement->GetMobility()];
+
+	LabelAndLayer += TEXT(" mobility=\"") + MobilityString + TEXT("\"");
+
 	return LabelAndLayer;
 }
 
diff --git a/Engine/Source/Runtime/Datasmith/DatasmithCore/Public/DatasmithDefinitions.h b/Engine/Source/Runtime/Datasmith/DatasmithCore/Public/DatasmithDefinitions.h
--- a/Engine/Source/Runtime/Datasmith/DatasmithCore/Public/DatasmithDefinitions.h
+++ b/Engine/Source/Runtime/Datasmith/DatasmithCore/Public/DatasmithDefinitions.h
@@ -647,3 +647,12 @@ static const TCHAR* KeyValuePropertyTypeStrings[] = { TEXT("String"), TEXT("Colo
 
 #define DATASMITH_DECALMATERIALNAME				TEXT("DecalMaterial")
 
+
+enum class EDatasmithActorMobilityType : uint8
+{
+	Static = 0,
+	Stationary,
+	Movable
+};
+
+static const TCHAR* DatasmithActorMobilityTypeStrings[] = { TEXT("Static"), TEXT("Stationary"), TEXT("Movable") };
diff --git a/Engine/Source/Runtime/Datasmith/DatasmithCore/Public/IDatasmithSceneElements.h b/Engine/Source/Runtime/Datasmith/DatasmithCore/Public/IDatasmithSceneElements.h
--- a/Engine/Source/Runtime/Datasmith/DatasmithCore/Public/IDatasmithSceneElements.h
+++ b/Engine/Source/Runtime/Datasmith/DatasmithCore/Public/IDatasmithSceneElements.h
@@ -150,6 +150,12 @@ public:
 
 	/** Get whether an actor's casts shadow */
 	virtual bool GetCastShadow() const = 0;
+
+	/** Set a mesh actor's visibility */
+	virtual void SetMobility(EDatasmithActorMobilityType Mobility) = 0;
+
+	/** Get a mesh actor's visibility */
+	virtual EDatasmithActorMobilityType GetMobility() const = 0;
 };
 
 /**
