Fixed rename crash on Datasmith reimport

#rnx

diff --git a/Engine/Plugins/Enterprise/DatasmithImporter/Source/DatasmithImporter/Private/Utility/DatasmithImporterImpl.cpp b/Engine/Plugins/Enterprise/DatasmithImporter/Source/DatasmithImporter/Private/Utility/DatasmithImporterImpl.cpp
--- a/Engine/Plugins/Enterprise/DatasmithImporter/Source/DatasmithImporter/Private/Utility/DatasmithImporterImpl.cpp
+++ b/Engine/Plugins/Enterprise/DatasmithImporter/Source/DatasmithImporter/Private/Utility/DatasmithImporterImpl.cpp
@@ -175,7 +175,29 @@ UObject* FDatasmithImporterImpl::PublicizeAsset( UObject* SourceAsset, const TCH
 		// If mesh's label has changed, update its name
 		if ( ExistingAsset->GetFName() != SourceAsset->GetFName() )
 		{
-			DestinationAsset->Rename( *AssetName, DestinationPackage, REN_DontCreateRedirectors | REN_NonTransactional );
+			if (DestinationAsset->Rename(*AssetName, DestinationPackage, REN_Test))
+			{
+				DestinationAsset->Rename(*AssetName, DestinationPackage, REN_DontCreateRedirectors | REN_NonTransactional);
+			}
+			// Something wrong happened during a previous re-import of this asset
+			// There already is a object in DestinationPackage. See UE-306372 for details.
+			// Correct things by making package and asset names equal as it should have always been.
+			else
+			{
+				const UObject* ExistingObject = StaticFindObject(nullptr, DestinationPackage, AssetName, EFindObjectFlags::ExactClass);
+				UE_LOG(LogDatasmithImport, Warning, TEXT("Renaming %s to %s failed because the latter already exists."), *DestinationAsset->GetFullName(), *ExistingObject->GetFullName());
+				
+				FAssetToolsModule& AssetToolsModule = FModuleManager::LoadModuleChecked<FAssetToolsModule>("AssetTools");
+
+				FString NewAssetName;
+				FString DestinationPackagePath;
+				AssetToolsModule.Get().CreateUniqueAssetName(DestinationPackage->GetPathName(), TEXT(""), DestinationPackagePath, NewAssetName);
+
+				DestinationPackage = CreatePackage(*DestinationPackagePath);
+				DestinationAsset->Rename(*NewAssetName, DestinationPackage, REN_DontCreateRedirectors | REN_NonTransactional);
+
+				UE_LOG(LogDatasmithImport, Log, TEXT("Asset has been renamed to %s."), *(DestinationPackagePath + TEXT(".") + NewAssetName));
+			}
 		}
 
 		if ( UStaticMesh* DestinationMesh = Cast< UStaticMesh >( DestinationAsset ) )
