Fixed crash of editor

#rnx

#Synced-CL 50905747

diff --git a/Engine/Plugins/Enterprise/DatasmithC4DImporter/Source/DatasmithC4DTranslator/Private/DatasmithC4DImporterMaterialSelector.cpp b/Engine/Plugins/Enterprise/DatasmithC4DImporter/Source/DatasmithC4DTranslator/Private/DatasmithC4DImporterMaterialSelector.cpp
deleted file mode 100644
diff --git a/Engine/Plugins/Enterprise/DatasmithC4DImporter/Source/DatasmithC4DTranslator/Private/DatasmithC4DImporterMaterialSelector.h b/Engine/Plugins/Enterprise/DatasmithC4DImporter/Source/DatasmithC4DTranslator/Private/DatasmithC4DImporterMaterialSelector.h
deleted file mode 100644
diff --git a/Engine/Plugins/Enterprise/DatasmithC4DImporter/Source/DatasmithC4DTranslator/Private/DatasmithC4DTranslatorModule.cpp b/Engine/Plugins/Enterprise/DatasmithC4DImporter/Source/DatasmithC4DTranslator/Private/DatasmithC4DTranslatorModule.cpp
--- a/Engine/Plugins/Enterprise/DatasmithC4DImporter/Source/DatasmithC4DTranslator/Private/DatasmithC4DTranslatorModule.cpp
+++ b/Engine/Plugins/Enterprise/DatasmithC4DImporter/Source/DatasmithC4DTranslator/Private/DatasmithC4DTranslatorModule.cpp
@@ -5,7 +5,6 @@
 #ifdef _MELANGE_SDK_
 
 #include "DatasmithC4DTranslator.h"
-#include "DatasmithC4DImporterMaterialSelector.h"
 #include "DatasmithC4DImportOptions.h"
 
 #include "ReferenceMaterials/DatasmithReferenceMaterialManager.h"
@@ -36,11 +35,6 @@ public:
 	/** IModuleInterface implementation */
 	virtual void StartupModule() override
 	{
-		// Make sure the DatasmithImporter module exists and has been initialized before adding FDatasmithC4DTranslator's material selector
-		FModuleManager::Get().LoadModule(TEXT("DatasmithTranslator"));
-
-		FDatasmithReferenceMaterialManager::Get().RegisterSelector(TEXT("C4DTranslator"), MakeShared< FDatasmithC4DImporterMaterialSelector >());
-
 	#if WITH_EDITOR
 		const FString EnvVariable = FPlatformMisc::GetEnvironmentVariable(TEXT("DATASMITHC4D_DEBUG"));
 		if (!EnvVariable.IsEmpty())
@@ -66,7 +60,6 @@ public:
 		if (FModuleManager::Get().IsModuleLoaded(TEXT("DatasmithTranslator")))
 		{
 			Datasmith::UnregisterTranslator<FDatasmithC4DTranslator>();
-			FDatasmithReferenceMaterialManager::Get().UnregisterSelector(TEXT("C4DTranslator"));
 		}
 	}
 };
diff --git a/Engine/Plugins/Enterprise/DatasmithImporter/Source/DatasmithTranslator/Private/DatasmithTranslatorModule.cpp b/Engine/Plugins/Enterprise/DatasmithImporter/Source/DatasmithTranslator/Private/DatasmithTranslatorModule.cpp
--- a/Engine/Plugins/Enterprise/DatasmithImporter/Source/DatasmithTranslator/Private/DatasmithTranslatorModule.cpp
+++ b/Engine/Plugins/Enterprise/DatasmithImporter/Source/DatasmithTranslator/Private/DatasmithTranslatorModule.cpp
@@ -2,6 +2,7 @@
 
 #include "DatasmithTranslatorModule.h"
 
+#include "ReferenceMaterials/DatasmithC4DImporterMaterialSelector.h"
 #include "ReferenceMaterials/DatasmithCityEngineMaterialSelector.h"
 #include "ReferenceMaterials/DatasmithReferenceMaterialManager.h"
 #include "ReferenceMaterials/DatasmithRevitMaterialSelector.h"
@@ -13,6 +14,7 @@ void IDatasmithTranslatorModule::StartupModule()
 	FDatasmithReferenceMaterialManager::Create();
 
 	//A minimal set of natively supported reference materials
+	FDatasmithReferenceMaterialManager::Get().RegisterSelector(TEXT("C4DTranslator"), MakeShared< FDatasmithC4DImporterMaterialSelector >());
 	FDatasmithReferenceMaterialManager::Get().RegisterSelector(TEXT("Revit"), MakeShared< FDatasmithRevitMaterialSelector >());
 	FDatasmithReferenceMaterialManager::Get().RegisterSelector(TEXT("SketchUp"), MakeShared< FDatasmithSketchUpMaterialSelector >());
 	FDatasmithReferenceMaterialManager::Get().RegisterSelector(TEXT("CityEngine"), MakeShared< FDatasmithCityEngineMaterialSelector >());
@@ -21,6 +23,7 @@ void IDatasmithTranslatorModule::StartupModule()
 
 void IDatasmithTranslatorModule::ShutdownModule()
 {
+	FDatasmithReferenceMaterialManager::Get().UnregisterSelector(TEXT("C4DTranslator"));
 	FDatasmithReferenceMaterialManager::Get().UnregisterSelector(TEXT("Revit"));
 	FDatasmithReferenceMaterialManager::Get().UnregisterSelector(TEXT("SketchUp"));
 	FDatasmithReferenceMaterialManager::Get().UnregisterSelector(TEXT("CityEngine"));
diff --git a/Engine/Plugins/Enterprise/DatasmithImporter/Source/DatasmithTranslator/Private/ReferenceMaterials/DatasmithC4DImporterMaterialSelector.cpp b/Engine/Plugins/Enterprise/DatasmithImporter/Source/DatasmithTranslator/Private/ReferenceMaterials/DatasmithC4DImporterMaterialSelector.cpp
new file mode 100644
--- /dev/null
+++ b/Engine/Plugins/Enterprise/DatasmithImporter/Source/DatasmithTranslator/Private/ReferenceMaterials/DatasmithC4DImporterMaterialSelector.cpp
@@ -0,0 +1,46 @@
+// Copyright Epic Games, Inc. All Rights Reserved.
+
+#include "DatasmithC4DImporterMaterialSelector.h"
+
+#include "IDatasmithSceneElements.h"
+
+#include "Materials/Material.h"
+#include "Materials/MaterialInstanceConstant.h"
+#include "UObject/SoftObjectPath.h"
+#include "Templates/Casts.h"
+
+FDatasmithC4DImporterMaterialSelector::FDatasmithC4DImporterMaterialSelector()
+{
+	ReferenceMaterial.FromSoftObjectPath( FSoftObjectPath("/DatasmithContent/Materials/C4DReference.C4DReference") );
+}
+
+bool FDatasmithC4DImporterMaterialSelector::IsValid() const
+{
+	return ReferenceMaterial.IsValid();
+}
+
+const FDatasmithReferenceMaterial& FDatasmithC4DImporterMaterialSelector::GetReferenceMaterial( const TSharedPtr< IDatasmithMaterialInstanceElement >& /*InDatasmithMaterial*/ ) const
+{
+	return ReferenceMaterial;
+}
+
+void FDatasmithC4DImporterMaterialSelector::FinalizeMaterialInstance(const TSharedPtr<IDatasmithMaterialInstanceElement>& InDatasmithMaterial, UMaterialInstanceConstant * MaterialInstance) const
+{
+	if (!InDatasmithMaterial.IsValid() || MaterialInstance == nullptr)
+	{
+		return;
+	}
+
+	// Set blend mode to translucent if material requires transparency.
+	if (InDatasmithMaterial->GetMaterialType() == EDatasmithReferenceMaterialType::Transparent)
+	{
+		MaterialInstance->BasePropertyOverrides.bOverride_BlendMode = true;
+		MaterialInstance->BasePropertyOverrides.BlendMode = BLEND_Translucent;
+	}
+	// Set blend mode to masked if material has cutouts.
+	else if (InDatasmithMaterial->GetMaterialType() == EDatasmithReferenceMaterialType::CutOut)
+	{
+		MaterialInstance->BasePropertyOverrides.bOverride_BlendMode = true;
+		MaterialInstance->BasePropertyOverrides.BlendMode = BLEND_Masked;
+	}
+}
\ No newline at end of file
diff --git a/Engine/Plugins/Enterprise/DatasmithImporter/Source/DatasmithTranslator/Private/ReferenceMaterials/DatasmithC4DImporterMaterialSelector.h b/Engine/Plugins/Enterprise/DatasmithImporter/Source/DatasmithTranslator/Private/ReferenceMaterials/DatasmithC4DImporterMaterialSelector.h
new file mode 100644
--- /dev/null
+++ b/Engine/Plugins/Enterprise/DatasmithImporter/Source/DatasmithTranslator/Private/ReferenceMaterials/DatasmithC4DImporterMaterialSelector.h
@@ -0,0 +1,26 @@
+// Copyright Epic Games, Inc. All Rights Reserved.
+
+#pragma once
+
+#include "DatasmithDefinitions.h"
+#include "ReferenceMaterials/DatasmithReferenceMaterial.h"
+#include "ReferenceMaterials/DatasmithReferenceMaterialSelector.h"
+
+#include "Templates/SharedPointer.h"
+
+class IDatasmithMaterialInstanceElement;
+
+class FDatasmithC4DImporterMaterialSelector : public FDatasmithReferenceMaterialSelector
+{
+public:
+	FDatasmithC4DImporterMaterialSelector();
+	virtual ~FDatasmithC4DImporterMaterialSelector() = default;
+
+	virtual bool IsValid() const override;
+	virtual const FDatasmithReferenceMaterial& GetReferenceMaterial( const TSharedPtr< IDatasmithMaterialInstanceElement >& /*InDatasmithMaterial*/ ) const override;
+	virtual void FinalizeMaterialInstance(const TSharedPtr< IDatasmithMaterialInstanceElement >& InDatasmithMaterial, UMaterialInstanceConstant* MaterialInstance) const override;
+
+private:
+
+	FDatasmithReferenceMaterial ReferenceMaterial;
+};
