Fix material slot name for Wire

#rnx

#Synced-CL 50905747

diff --git a/Engine/Plugins/Enterprise/DatasmithImporter/Source/DatasmithImporter/Private/DatasmithStaticMeshImporter.cpp b/Engine/Plugins/Enterprise/DatasmithImporter/Source/DatasmithImporter/Private/DatasmithStaticMeshImporter.cpp
--- a/Engine/Plugins/Enterprise/DatasmithImporter/Source/DatasmithImporter/Private/DatasmithStaticMeshImporter.cpp
+++ b/Engine/Plugins/Enterprise/DatasmithImporter/Source/DatasmithImporter/Private/DatasmithStaticMeshImporter.cpp
@@ -764,6 +764,7 @@ void FDatasmithStaticMeshImporter::ApplyMaterialsToStaticMesh( const FDatasmithA
 
 		if ( MaterialTemplate )
 		{
+			MaterialTemplate->MaterialSlotName = MaterialInterface  ? *MaterialInterface->GetName() : TEXT("Undefined");
 			MaterialTemplate->MaterialInterface = MaterialInterface;
 		}
 	}
