clean up code changing texture properties with wrapping in PreEditChange/PostEditChange as required.
Most of these are not actual bugs, because the texture is newly created, or was using an alternative form that works but is not recommended (eg. Modify/UpdateResource).
As much as possible prefer sticking to the standard pattern of PreEdit/PostEdit.
Some of these may have been real bugs causing problems.  That is unconfirmed.

#Synced-CL 51487686

diff --git a/Engine/Plugins/2D/Paper2D/Source/Paper2DEditor/Private/Atlasing/PaperAtlasGenerator.cpp b/Engine/Plugins/2D/Paper2D/Source/Paper2DEditor/Private/Atlasing/PaperAtlasGenerator.cpp
--- a/Engine/Plugins/2D/Paper2D/Source/Paper2DEditor/Private/Atlasing/PaperAtlasGenerator.cpp
+++ b/Engine/Plugins/2D/Paper2D/Source/Paper2DEditor/Private/Atlasing/PaperAtlasGenerator.cpp
@@ -195,6 +195,7 @@ void FPaperAtlasGenerator::HandleAssetChangedEvent(UPaperSpriteAtlas* Atlas)
 
 		const bool bAtlasDirty = RemappedAtlasForceDirty[AtlasIndex] || (AtlasTexture->GetImportedSize() != FIntPoint(AtlasWidth, AtlasHeight));
 
+		AtlasTexture->PreEditChange(nullptr);
 		// Propagate texture settings
 		AtlasTexture->CompressionSettings = Atlas->CompressionSettings;
 		AtlasTexture->Filter = Atlas->Filter;
@@ -235,7 +236,6 @@ void FPaperAtlasGenerator::HandleAssetChangedEvent(UPaperSpriteAtlas* Atlas)
 
 
 		AtlasTexture->Source.Init(AtlasWidth, AtlasHeight, /*NumSlices=*/ 1, Atlas->MipCount, ETextureSourceFormat::TSF_BGRA8, AtlasTextureData.GetData());
-		AtlasTexture->UpdateResource();
 		AtlasTexture->PostEditChange();
 	}
 
diff --git a/Engine/Plugins/Bridge/Source/MegascansPlugin/Private/Tools/BlendMaterials.cpp b/Engine/Plugins/Bridge/Source/MegascansPlugin/Private/Tools/BlendMaterials.cpp
--- a/Engine/Plugins/Bridge/Source/MegascansPlugin/Private/Tools/BlendMaterials.cpp
+++ b/Engine/Plugins/Bridge/Source/MegascansPlugin/Private/Tools/BlendMaterials.cpp
@@ -124,16 +124,15 @@ bool FMaterialBlend::ValidateSelectedAssets(TArray<FString> SelectedMaterials, F
 
 void FMaterialBlend::HandleTextureLoading(FAssetData TextureData)
 {
-
     UTexture *TextureAsset = Cast<UTexture>(TextureData.GetAsset());
-    uint8 VT = TextureAsset->VirtualTextureStreaming;
-
-    TextureAsset->VirtualTextureStreaming = 1;
+	check( TextureAsset != nullptr );
+	if ( ! TextureAsset->VirtualTextureStreaming )
+	{
+		TextureAsset->PreEditChange(nullptr);
+	    TextureAsset->VirtualTextureStreaming = 1;
+	}
     TextureAsset->SetFlags(RF_Standalone);
-    TextureAsset->MarkPackageDirty();
     TextureAsset->PostEditChange();
-
-    // AssetUtils::SavePackage(TextureAsset);
 }
 
 void FMaterialBlend::ConvertToVirtualTextures(FUAssetMeta AssetMetaData)
diff --git a/Engine/Plugins/Enterprise/DatasmithImporter/Source/DatasmithImporter/Private/DatasmithTextureImporter.cpp b/Engine/Plugins/Enterprise/DatasmithImporter/Source/DatasmithImporter/Private/DatasmithTextureImporter.cpp
--- a/Engine/Plugins/Enterprise/DatasmithImporter/Source/DatasmithImporter/Private/DatasmithTextureImporter.cpp
+++ b/Engine/Plugins/Enterprise/DatasmithImporter/Source/DatasmithImporter/Private/DatasmithTextureImporter.cpp
@@ -267,6 +267,8 @@ UTexture* FDatasmithTextureImporter::CreateTexture(const TSharedPtr<IDatasmithTe
 															*Extension, PtrTextureData, PtrTextureDataEnd, GWarn));
 	if (Texture != nullptr)
 	{
+		Texture->PreEditChange(nullptr);
+
 		static_assert(TextureAddress::TA_Wrap == (int)EDatasmithTextureAddress::Wrap && TextureAddress::TA_Mirror == (int)EDatasmithTextureAddress::Mirror, "Texture Address enum doesn't match!" );
 
 		FMD5Hash Hash = TextureElement->CalculateElementHash(false);
@@ -290,17 +292,11 @@ UTexture* FDatasmithTextureImporter::CreateTexture(const TSharedPtr<IDatasmithTe
 			break;
 		}
 
-		bool bUpdateResource = false;
-
-		bUpdateResource |= Texture->Filter != TexFilter;
 		Texture->Filter = TexFilter;
 
 		if (UTexture2D* Texture2D = Cast<UTexture2D>(Texture))
 		{
-			bUpdateResource |= Texture2D->AddressX != (TextureAddress)TextureElement->GetTextureAddressX();
 			Texture2D->AddressX = (TextureAddress)TextureElement->GetTextureAddressX();
-
-			bUpdateResource |= Texture2D->AddressY != (TextureAddress)TextureElement->GetTextureAddressY();
 			Texture2D->AddressY = (TextureAddress)TextureElement->GetTextureAddressY();
 		}
 
@@ -313,7 +309,6 @@ UTexture* FDatasmithTextureImporter::CreateTexture(const TSharedPtr<IDatasmithTe
 		if (FMath::IsNearlyEqual(RGBCurve, 1.0f) == false && RGBCurve > 0.f)
 		{
 			Texture->AdjustRGBCurve = RGBCurve;
-			bUpdateResource = true;
 		}
 
 		// Double check that the texture color space matches what we requested
@@ -321,22 +316,13 @@ UTexture* FDatasmithTextureImporter::CreateTexture(const TSharedPtr<IDatasmithTe
 		if (!Texture->SRGB && ColorSpace == EDatasmithColorSpace::sRGB)
 		{
 			Texture->SRGB = true;
-			bUpdateResource = true;
 		}
 		else if (Texture->SRGB && ColorSpace == EDatasmithColorSpace::Linear)
 		{
 			Texture->SRGB = false;
-			bUpdateResource = true;
 		}
 
-		if (bUpdateResource)
-		{
-			// Make sure the previous update done by the factory has been completed
-			FlushRenderingCommands();
-
-			Texture->UpdateResource();
-		}
-		Texture->MarkPackageDirty();
+		Texture->PostEditChange();
 	}
 
 	return Texture;
diff --git a/Engine/Plugins/Enterprise/GLTFExporter/Source/GLTFExporter/Private/Converters/GLTFMaterialProxyFactory.cpp b/Engine/Plugins/Enterprise/GLTFExporter/Source/GLTFExporter/Private/Converters/GLTFMaterialProxyFactory.cpp
--- a/Engine/Plugins/Enterprise/GLTFExporter/Source/GLTFExporter/Private/Converters/GLTFMaterialProxyFactory.cpp
+++ b/Engine/Plugins/Enterprise/GLTFExporter/Source/GLTFExporter/Private/Converters/GLTFMaterialProxyFactory.cpp
@@ -241,9 +241,13 @@ UTexture2D* FGLTFMaterialProxyFactory::CreateTexture(const FGLTFImageData* Image
 	UPackage* Package = FindOrCreatePackage(BaseName);
 
 	UTexture2D* Texture = FImageUtils::CreateTexture2D(ImageData->Size.X, ImageData->Size.Y, *ImageData->Pixels, Package, BaseName, RF_Public | RF_Standalone, TexParams);
+	
+	Texture->PreEditChange(nullptr);
 	Texture->Filter = ConvertFilter(JsonSampler.MagFilter);
 	Texture->AddressX = ConvertWrap(JsonSampler.WrapS);
 	Texture->AddressY = ConvertWrap(JsonSampler.WrapT);
+	Texture->PostEditChange();
+
 	return Texture;
 }
 
diff --git a/Engine/Plugins/Experimental/Water/Source/Editor/Private/WaterBrushManager.cpp b/Engine/Plugins/Experimental/Water/Source/Editor/Private/WaterBrushManager.cpp
--- a/Engine/Plugins/Experimental/Water/Source/Editor/Private/WaterBrushManager.cpp
+++ b/Engine/Plugins/Experimental/Water/Source/Editor/Private/WaterBrushManager.cpp
@@ -334,7 +334,10 @@ void AWaterBrushManager::OnCurveUpdated(UCurveBase* Curve, EPropertyChangeType::
 
 void AWaterBrushManager::BlueprintOnRenderTargetTexturesUpdated_Native(UTexture2D* VelocityTexture)
 {
+	VelocityTexture->PreEditChange(nullptr);
 	VelocityTexture->LODBias = 0;
+	VelocityTexture->PostEditChange();
+
 	UseDynamicPreviewRT = false;
 }
 
diff --git a/Engine/Plugins/Mutable/Source/CustomizableObjectEditor/Private/MuCOE/UnrealBakeHelpers.cpp b/Engine/Plugins/Mutable/Source/CustomizableObjectEditor/Private/MuCOE/UnrealBakeHelpers.cpp
--- a/Engine/Plugins/Mutable/Source/CustomizableObjectEditor/Private/MuCOE/UnrealBakeHelpers.cpp
+++ b/Engine/Plugins/Mutable/Source/CustomizableObjectEditor/Private/MuCOE/UnrealBakeHelpers.cpp
@@ -134,6 +134,7 @@ UTexture2D* FUnrealBakeHelpers::BakeHelper_CreateAssetTexture(UTexture2D* Source
 	// action is not used often.
 	DupTex->AddToRoot();
 	DupTex->MarkPackageDirty();
+	DupTex->PreEditChange(nullptr);
 
 	// Notify the asset registry
 	FAssetRegistryModule::AssetCreated(DupTex);
@@ -156,16 +157,18 @@ UTexture2D* FUnrealBakeHelpers::BakeHelper_CreateAssetTexture(UTexture2D* Source
 	{
 		DupTex->CompressionSettings = OrgTex->CompressionSettings;
 	}
-
+	
 	// Mutable textures only have platform data. We need to build the source data for them to be assets.
-	DupTex->Source.Init(SourceTextureSizeX, SourceTextureSizeY, 1, 1, PixelFormat);
-
 	int32 MipCount = SourceTexture->GetPlatformData()->Mips.Num();
 	if (!MipCount)
 	{
 		UE_LOG(LogMutable, Warning, TEXT("Bake Instances: Empty texture found [%s]."), *SourceTexture->GetName());
+		
+		DupTex->PostEditChange();
 		return DupTex;
 	}
+	
+	DupTex->Source.Init(SourceTextureSizeX, SourceTextureSizeY, 1, 1, PixelFormat);
 
 	// Create a mutable image from the platform data.
 	UE::Mutable::Private::EImageFormat MutableFormat = UnrealToMutablePixelFormat(SourceTexture->GetPlatformData()->PixelFormat, SourceTexture->HasAlphaChannel());
@@ -261,7 +264,7 @@ UTexture2D* FUnrealBakeHelpers::BakeHelper_CreateAssetTexture(UTexture2D* Source
 		DupTex->SetLayerFormatSettings(0, Settings);
 	}
 
-	DupTex->UpdateResource();
+	DupTex->PostEditChange();
 
 	return DupTex;
 }
diff --git a/Engine/Plugins/Runtime/HairStrands/Source/HairStrandsCore/Private/GroomTextureBuilder.cpp b/Engine/Plugins/Runtime/HairStrands/Source/HairStrandsCore/Private/GroomTextureBuilder.cpp
--- a/Engine/Plugins/Runtime/HairStrands/Source/HairStrandsCore/Private/GroomTextureBuilder.cpp
+++ b/Engine/Plugins/Runtime/HairStrands/Source/HairStrandsCore/Private/GroomTextureBuilder.cpp
@@ -55,6 +55,10 @@ void FGroomTextureBuilder::AllocateFollicleTextureResources(UTexture2D* Out, con
 	{
 		return;
 	}
+	
+#if WITH_EDITOR
+	Out->PreEditChange(nullptr);
+#endif
 
 	FTextureFormatSettings FormatSettings;
 	FormatSettings.CompressionNone = true;
@@ -631,6 +635,10 @@ static void AddTextureDilationPass(
 
 static void InternalAllocateStrandsTexture(UTexture2D* Out, const FIntPoint& Resolution, uint32 MipCount, EPixelFormat Format, ETextureSourceFormat SourceFormat)
 {
+#if WITH_EDITOR
+	Out->PreEditChange(nullptr);
+#endif
+
 	FTextureFormatSettings FormatSettings;
 	FormatSettings.SRGB = false;
 	FormatSettings.CompressionNone = true;
diff --git a/Engine/Plugins/Runtime/OpenCV/Source/OpenCVHelper/Private/OpenCVHelper.cpp b/Engine/Plugins/Runtime/OpenCV/Source/OpenCVHelper/Private/OpenCVHelper.cpp
--- a/Engine/Plugins/Runtime/OpenCV/Source/OpenCVHelper/Private/OpenCVHelper.cpp
+++ b/Engine/Plugins/Runtime/OpenCV/Source/OpenCVHelper/Private/OpenCVHelper.cpp
@@ -208,6 +208,8 @@ UTexture2D* FOpenCVHelper::TextureFromCvMat(cv::Mat& Mat, const FString* Package
 			return nullptr;
 		}
 
+		Texture->PreEditChange(nullptr);
+
 		const int32 NumSlices = 1;
 		const int32 NumMips = 1;
 
@@ -250,9 +252,7 @@ UTexture2D* FOpenCVHelper::TextureFromCvMat(cv::Mat& Mat, const FString* Package
 		Texture->GetPlatformData()->SizeY = Mat.rows;
 		Texture->GetPlatformData()->PixelFormat = PixelFormat;
 
-		Texture->UpdateResource();
-
-		Texture->MarkPackageDirty();
+		Texture->PostEditChange();
 	}
 	else
 #endif //WITH_EDITOR
diff --git a/Engine/Source/Runtime/Engine/Private/Font.cpp b/Engine/Source/Runtime/Engine/Private/Font.cpp
--- a/Engine/Source/Runtime/Engine/Private/Font.cpp
+++ b/Engine/Source/Runtime/Engine/Private/Font.cpp
@@ -55,24 +55,45 @@ void UFont::PostLoad()
 
 	// Cache the character count and the maximum character height for this font
 	CacheCharacterCountAndMaxCharHeight();
-
+	
+#if WITH_EDITOR
 	for( int32 TextureIndex = 0 ; TextureIndex < Textures.Num() ; ++TextureIndex )
 	{
 		UTexture2D* Texture = Textures[TextureIndex];
 		if( Texture )
 		{	
 			Texture->SetFlags(RF_Public);
-			Texture->LODGroup = TEXTUREGROUP_UI;
-			Texture->ConditionalPostLoad();
+			
+			// Texture changes like this should NOT be done in PostLoad
+
+			bool bTextureChanged = false;
+			if ( Texture->LODGroup != TEXTUREGROUP_UI )
+			{
+				bTextureChanged = true;
+				Texture->PreEditChange(nullptr);
+
+				Texture->LODGroup = TEXTUREGROUP_UI;
+			}
 
 			// Fix up compression type for distance field fonts.
 			if (Texture->CompressionSettings == TC_Displacementmap && !Texture->SRGB)
 			{
+				if ( ! bTextureChanged )
+				{
+					bTextureChanged = true;
+					Texture->PreEditChange(nullptr);
+				}
+
 				Texture->CompressionSettings = TC_DistanceFieldFont;
-				Texture->UpdateResource();
+			}
+
+			if ( bTextureChanged )
+			{
+				Texture->PostEditChange();
 			}
 		}
 	}
+#endif
 
 #if WITH_EDITORONLY_DATA
 	if (FontCacheType == EFontCacheType::Runtime)
